RigLogic AnimNext

* Added new RigLogicAnimNext plugin.
* Added new RigLogic trait and task.
* RigLogic instance data pool that automatically grows with the maximum number of parallel evaluations. Data is shared between skel meshes.

diff --git a/Engine/Plugins/Experimental/AnimNext/Source/AnimNext/Private/EvaluationVM/EvaluationVM.cpp b/Engine/Plugins/Experimental/AnimNext/Source/AnimNext/Private/EvaluationVM/EvaluationVM.cpp
--- a/Engine/Plugins/Experimental/AnimNext/Source/AnimNext/Private/EvaluationVM/EvaluationVM.cpp
+++ b/Engine/Plugins/Experimental/AnimNext/Source/AnimNext/Private/EvaluationVM/EvaluationVM.cpp
@@ -33,8 +33,6 @@ namespace UE::AnimNext
 		, CurrentLOD(InCurrentLOD)
 		, EvaluationFlags(InEvaluationFlags)
 	{
-		CurveFilter.SetFilterMode(UE::Anim::ECurveFilterMode::DisallowAll);
-
 		//const UE::Anim::FCurveFilterSettings CurveFilterSettings(UE::Anim::ECurveFilterMode::DisallowAll);
 		//if(USkeleton* Skeleton = const_cast<USkeleton*>(InReferencePose.Skeleton.Get())) // const_cast because the bone container takes a mutable reference
 		{
diff --git a/Engine/Plugins/Experimental/RigLogicAnimNext/RigLogicAnimNext.uplugin b/Engine/Plugins/Experimental/RigLogicAnimNext/RigLogicAnimNext.uplugin
new file mode 100644
--- /dev/null
+++ b/Engine/Plugins/Experimental/RigLogicAnimNext/RigLogicAnimNext.uplugin
@@ -0,0 +1,40 @@
+{
+	"FileVersion": 3,
+	"Version": 1,
+	"VersionName": "1.0",
+	"FriendlyName": "RigLogic for AnimNext",
+	"Description": "RigLogic for AnimNext",
+	"Category": "Animation",
+	"CreatedBy": "Epic Games Inc.",
+	"CreatedByURL": "",
+	"DocsURL": "",
+	"MarketplaceURL": "",
+	"SupportURL": "",
+	"CanContainContent": true,
+	"IsBetaVersion": false,
+	"IsExperimentalVersion": true,
+	"Installed": false,
+	"Modules":
+	[
+		{
+			"Name": "RigLogicAnimNext",
+			"Type": "Runtime",
+			"LoadingPhase": "Default"
+		}
+	],
+	"Plugins":
+	[
+		{
+			"Name": "RigLogic",
+			"Enabled": true
+		},
+		{
+			"Name": "AnimNext",
+			"Enabled": true
+		},
+		{
+			"Name": "AnimNextAnimGraph",
+			"Enabled": true
+		}
+	]
+}
\ No newline at end of file
diff --git a/Engine/Plugins/Experimental/RigLogicAnimNext/Source/RigLogicAnimNext/Private/RigLogicAnimNext.cpp b/Engine/Plugins/Experimental/RigLogicAnimNext/Source/RigLogicAnimNext/Private/RigLogicAnimNext.cpp
new file mode 100644
--- /dev/null
+++ b/Engine/Plugins/Experimental/RigLogicAnimNext/Source/RigLogicAnimNext/Private/RigLogicAnimNext.cpp
@@ -0,0 +1,19 @@
+// Copyright Epic Games, Inc. All Rights Reserved.
+
+#include "RigLogicAnimNext.h"
+
+#define LOCTEXT_NAMESPACE "FRigLogicAnimNextModule"
+
+DEFINE_LOG_CATEGORY(LogRigLogicAnimNext);
+
+void FRigLogicAnimNextModule::StartupModule()
+{
+}
+
+void FRigLogicAnimNextModule::ShutdownModule()
+{
+}
+
+#undef LOCTEXT_NAMESPACE
+	
+IMPLEMENT_MODULE(FRigLogicAnimNextModule, RigLogicAnimNext)
\ No newline at end of file
diff --git a/Engine/Plugins/Experimental/RigLogicAnimNext/Source/RigLogicAnimNext/Private/RigLogicInstanceData.cpp b/Engine/Plugins/Experimental/RigLogicAnimNext/Source/RigLogicAnimNext/Private/RigLogicInstanceData.cpp
new file mode 100644
--- /dev/null
+++ b/Engine/Plugins/Experimental/RigLogicAnimNext/Source/RigLogicAnimNext/Private/RigLogicInstanceData.cpp
@@ -0,0 +1,101 @@
+// Copyright Epic Games, Inc. All Rights Reserved.
+
+#include "RigLogicInstanceData.h"
+#include "RigLogicAnimNext.h"
+
+#include "Engine/SkeletalMesh.h"
+#include "DNAAsset.h"
+#include "LODPose.h"
+#include "SharedRigRuntimeContext.h"
+
+namespace UE::AnimNext
+{
+	void FRigLogicAnimNextInstanceData::Init(USkeletalMesh* SkeletalMesh)
+	{
+		const USkeleton* Skeleton = SkeletalMesh->GetSkeleton();
+		if (!Skeleton)
+		{
+			UE_LOG(LogRigLogicAnimNext, Warning, TEXT("No skeleton assigned to the skeletal mesh."));
+			return;
+		}
+
+		UDNAAsset* DNAAsset = Cast<UDNAAsset>(SkeletalMesh->GetAssetUserDataOfClass(UDNAAsset::StaticClass()));
+		if (!DNAAsset)
+		{
+			UE_LOG(LogRigLogicAnimNext, Warning, TEXT("No DNA asset assigned to the skeletal mesh."));
+			return;
+		}
+
+		TSharedPtr<FSharedRigRuntimeContext> SharedRigRuntimeContext = DNAAsset->GetRigRuntimeContext();
+		if (!SharedRigRuntimeContext.IsValid())
+		{
+			UE_LOG(LogRigLogicAnimNext, Warning, TEXT("Can't get the shared rig runtime context."));
+			return;
+		}
+
+		if (CachedRigRuntimeContext != SharedRigRuntimeContext)
+		{
+			CachedRigRuntimeContext = SharedRigRuntimeContext;
+			RigInstance = TUniquePtr<FRigInstance>(new FRigInstance(CachedRigRuntimeContext->RigLogic.Get()));
+		}
+
+		CachedDNAIndexMapping = DNAAsset->GetDNAIndexMapping(Skeleton, SkeletalMesh);
+
+		InitBoneIndexMapping();
+		InitSparseAndDenseDriverJointMapping();
+	}
+
+	void FRigLogicAnimNextInstanceData::InitBoneIndexMapping()
+	{
+		const uint32 NumLODs = CachedRigRuntimeContext->VariableJointIndicesPerLOD.Num();
+		RigLogicToSkeletonBoneIndexMappingPerLOD.Empty();
+		RigLogicToSkeletonBoneIndexMappingPerLOD.SetNum(NumLODs);
+
+		for (uint32 LODLevel = 0; LODLevel < NumLODs; ++LODLevel)
+		{
+			const TArray<uint16>& VariableJointIndices = CachedRigRuntimeContext->VariableJointIndicesPerLOD[LODLevel].Values;
+			RigLogicToSkeletonBoneIndexMappingPerLOD[LODLevel].Reserve(VariableJointIndices.Num());
+
+			for (const uint16 RigLogicJointIndex : VariableJointIndices)
+			{
+				const FMeshPoseBoneIndex MeshPoseBoneIndex = CachedDNAIndexMapping->JointsMapDNAIndicesToMeshPoseBoneIndices[RigLogicJointIndex];
+				if (MeshPoseBoneIndex != INDEX_NONE)
+				{
+					const int32 SkeletonBoneIndex = MeshPoseBoneIndex.GetInt();
+					RigLogicToSkeletonBoneIndexMappingPerLOD[LODLevel].Add({ RigLogicJointIndex, SkeletonBoneIndex });
+				}
+				else
+				{
+					UE_LOG(LogTemp, Warning, TEXT("Could not find bone in skeleton for RigLogic joint with index %i."), RigLogicJointIndex);
+				}
+			}
+		}
+	}
+
+	void FRigLogicAnimNextInstanceData::InitSparseAndDenseDriverJointMapping()
+	{
+		// Populate driver joint to raw control attribute mapping (used to feed RigLogic with inputs from the joint hierarchy)
+		SparseDriverJointsToControlAttributesMap.Empty();
+		DenseDriverJointsToControlAttributesMap.Empty();
+		DenseDriverJointsToControlAttributesMap.Reserve(CachedDNAIndexMapping->DriverJointsToControlAttributesMap.Num());
+
+		// Sparse mapping will likely remain empty so no reservation happens
+		for (const auto& Mapping : CachedDNAIndexMapping->DriverJointsToControlAttributesMap)
+		{
+			const FMeshPoseBoneIndex MeshPoseBoneIndex = Mapping.MeshPoseBoneIndex;
+
+			if (MeshPoseBoneIndex != INDEX_NONE)
+			{
+				const int32 SkeletonBoneIndex = MeshPoseBoneIndex.GetInt();
+				if ((Mapping.RotationX != INDEX_NONE) && (Mapping.RotationY != INDEX_NONE) && (Mapping.RotationZ != INDEX_NONE) && (Mapping.RotationW != INDEX_NONE))
+				{
+					DenseDriverJointsToControlAttributesMap.Add({ SkeletonBoneIndex, Mapping.DNAJointIndex, Mapping.RotationX, Mapping.RotationY, Mapping.RotationZ, Mapping.RotationW });
+				}
+				else
+				{
+					SparseDriverJointsToControlAttributesMap.Add({ SkeletonBoneIndex, Mapping.DNAJointIndex, Mapping.RotationX, Mapping.RotationY, Mapping.RotationZ, Mapping.RotationW });
+				}
+			}
+		}
+	}
+} // namespace UE::AnimNext
\ No newline at end of file
diff --git a/Engine/Plugins/Experimental/RigLogicAnimNext/Source/RigLogicAnimNext/Private/RigLogicInstanceDataPool.cpp b/Engine/Plugins/Experimental/RigLogicAnimNext/Source/RigLogicAnimNext/Private/RigLogicInstanceDataPool.cpp
new file mode 100644
--- /dev/null
+++ b/Engine/Plugins/Experimental/RigLogicAnimNext/Source/RigLogicAnimNext/Private/RigLogicInstanceDataPool.cpp
@@ -0,0 +1,93 @@
+// Copyright Epic Games, Inc. All Rights Reserved.
+
+#include "RigLogicInstanceDataPool.h"
+#include "RigLogicAnimNext.h"
+
+namespace UE::AnimNext
+{
+	TSharedPtr<FRigLogicAnimNextInstanceData> FRigLogicInstanceDataPool::RequestData(const UE::AnimNext::FReferencePose* InReferencePose)
+	{
+		// Only lock while searching
+		{
+			FScopeLock Lock(&PoolAccessCriticalSection);
+
+			TArray<TSharedPtr<FRigLogicAnimNextInstanceData>>* Found = Datas.Find(InReferencePose->SkeletalMesh);
+			if (Found)
+			{
+				if (!Found->IsEmpty())
+				{
+					// We have pre-allocated data already, use one that is already there.
+					return Found->Pop();
+				}
+			}
+		}
+		// We don't need to keep things locked here as we're cerating a new object
+		// Initializing a new data is not a lightweight operation, so be sure to close the lock first.
+
+		// We don't have any pool datas in the array yet. Create a new one.
+		TSharedPtr<FRigLogicAnimNextInstanceData> NewObject = TSharedPtr<FRigLogicAnimNextInstanceData>(new FRigLogicAnimNextInstanceData());
+		NewObject->Init(const_cast<USkeletalMesh*>(InReferencePose->SkeletalMesh.Get()));
+		return NewObject;
+	}
+
+	void FRigLogicInstanceDataPool::FreeData(TWeakObjectPtr<const USkeletalMesh> SkeletalMesh, TSharedPtr<FRigLogicAnimNextInstanceData> InData)
+	{
+		FScopeLock Lock(&PoolAccessCriticalSection);
+
+		// Do we have an element for the given skel mesh?
+		TArray<TSharedPtr<FRigLogicAnimNextInstanceData>>* Found = Datas.Find(SkeletalMesh);
+		if (!Found)
+		{
+			// We don't. Create a new one and add the instance data to it.
+			TArray<TSharedPtr<FRigLogicAnimNextInstanceData>> NewArray;
+			NewArray.Add(InData);
+			Datas.Add(SkeletalMesh, NewArray);
+		}
+		else
+		{
+			// Found one, just add the instance data back in.
+			Found->Add(InData);
+		}
+	}
+
+	void FRigLogicInstanceDataPool::GarbageCollect()
+	{
+		FScopeLock Lock(&PoolAccessCriticalSection);
+
+		for (auto Iterator = Datas.CreateIterator(); Iterator; ++Iterator)
+		{
+			// Remove cached instance datas for skeletal meshes that are not loaded anymore.
+			if (!Iterator->Key.IsValid())
+			{
+				Iterator.RemoveCurrent();
+			}
+		}
+	}
+
+	void FRigLogicInstanceDataPool::Log()
+	{
+		FScopeLock Lock(&PoolAccessCriticalSection);
+
+		UE_LOG(LogRigLogicAnimNext, Display, TEXT("Pool data:"));
+		uint32 Counter = 0;
+		for (auto Iterator = Datas.CreateIterator(); Iterator; ++Iterator)
+		{
+			if (Iterator->Key.IsValid())
+			{
+				const USkeletalMesh* SkeletalMesh = Iterator->Key.Get();
+				UE_LOG(LogRigLogicAnimNext, Display, TEXT("   - Skeletal Mesh %s:"), *SkeletalMesh->GetPathName());
+
+				const TArray<TSharedPtr<FRigLogicAnimNextInstanceData>>& InstanceDataArray = Iterator->Value;
+				for (int32 i = 0; i < InstanceDataArray.Num(); ++i)
+				{
+					UE_LOG(LogRigLogicAnimNext, Display, TEXT("      + InstanceData %i (0x%d):"), i, InstanceDataArray[i].Get());
+				}
+			}
+			else
+			{
+				UE_LOG(LogRigLogicAnimNext, Warning, TEXT("Entry %i linked to an invalid skeletal mesh."), Counter);
+			}
+			Counter++;
+		}
+	}
+} // namespace UE::AnimNext
\ No newline at end of file
diff --git a/Engine/Plugins/Experimental/RigLogicAnimNext/Source/RigLogicAnimNext/Private/RigLogicTask.cpp b/Engine/Plugins/Experimental/RigLogicAnimNext/Source/RigLogicAnimNext/Private/RigLogicTask.cpp
new file mode 100644
--- /dev/null
+++ b/Engine/Plugins/Experimental/RigLogicAnimNext/Source/RigLogicAnimNext/Private/RigLogicTask.cpp
@@ -0,0 +1,246 @@
+// Copyright Epic Games, Inc. All Rights Reserved.
+
+#include "RigLogicTask.h"
+
+#include "EvaluationVM/EvaluationVM.h"
+#include "EvaluationVM/KeyframeState.h"
+#include "TransformArrayOperations.h"
+
+#include "RigLogicAnimNext.h"
+#include "DNAIndexMapping.h"
+#include "RigLogic.h"
+#include "RigInstance.h"
+#include "SharedRigRuntimeContext.h"
+#include "RigLogicInstanceDataPool.h"
+
+#include <tdm/TDM.h>
+
+DEFINE_STAT(STAT_AnimNext_Task_RigLogic);
+
+FAnimNextRigLogicTask FAnimNextRigLogicTask::Make()
+{
+	return {};
+}
+
+void FAnimNextRigLogicTask::Execute(UE::AnimNext::FEvaluationVM& VM) const
+{
+	SCOPE_CYCLE_COUNTER(STAT_AnimNext_Task_RigLogic);
+
+	// Pop pose, we'll re-use it for our result
+	TUniquePtr<UE::AnimNext::FKeyframeState> Keyframe;
+	if (!VM.PopValue(UE::AnimNext::KEYFRAME_STACK_NAME, Keyframe))
+	{
+		// We have no inputs, nothing to do.
+		return;
+	}
+
+	// Lazy-init instance data.
+	FRigLogicAnimNextModule& Module = FModuleManager::GetModuleChecked<FRigLogicAnimNextModule>("RigLogicAnimNext");
+	TSharedPtr<UE::AnimNext::FRigLogicAnimNextInstanceData> NewInstanceData = Module.DataPool.RequestData(Keyframe->Pose.RefPose);
+
+	FRigInstance* RigInstance = NewInstanceData->RigInstance.Get();
+	const FSharedRigRuntimeContext* RigRuntimeContext = NewInstanceData->CachedRigRuntimeContext.Get();
+	const FDNAIndexMapping* DNAIndexMapping = NewInstanceData->CachedDNAIndexMapping.Get();
+	if (RigInstance && RigRuntimeContext && DNAIndexMapping)
+	{
+		if (const FRigLogic* RigLogic = RigRuntimeContext->RigLogic.Get())
+		{
+			// 1. Feed input curves to RigLogic.
+			FBlendedCurve& Curves = Keyframe->Curves;
+			const TArrayView<const float> NeutralJointValues = NewInstanceData->CachedRigRuntimeContext->RigLogic->GetNeutralJointValues();
+
+			UpdateControlCurves(RigInstance,
+				Curves,
+				DNAIndexMapping->ControlAttributeCurves,
+				DNAIndexMapping->NeuralNetworkMaskCurves,
+				NeutralJointValues,
+				NewInstanceData->SparseDriverJointsToControlAttributesMap,
+				NewInstanceData->DenseDriverJointsToControlAttributesMap,
+				Keyframe->Pose);
+
+			// 2. Evaluate RigLogic.
+			const int32 LODLevel = Keyframe->Pose.LODLevel;
+			RigInstance->SetLOD(LODLevel);
+			RigLogic->Calculate(RigInstance);
+
+			// 3. Update AnimNext bone transforms and curves with RigLogic output.
+			TArrayView<const uint16> VariableJointIndices = NewInstanceData->CachedRigRuntimeContext->VariableJointIndicesPerLOD[LODLevel].Values;
+			const TArrayView<const UE::AnimNext::FRigLogicBoneMapping> RigLogicIndexToMeshIndexMapping = NewInstanceData->RigLogicToSkeletonBoneIndexMappingPerLOD[LODLevel];
+			UpdateJoints(RigLogicIndexToMeshIndexMapping,
+				VariableJointIndices,
+				NeutralJointValues,
+				RigInstance->GetJointOutputs(),
+				Keyframe->Pose);
+
+			UpdateBlendShapeCurves(DNAIndexMapping->MorphTargetCurvesPerLOD[LODLevel], RigInstance->GetBlendShapeOutputs(), Curves);
+			UpdateAnimMapCurves(DNAIndexMapping->MaskMultiplierCurvesPerLOD[LODLevel], RigInstance->GetAnimatedMapOutputs(), Curves);
+		}
+	}
+
+	Module.DataPool.FreeData(Keyframe->Pose.RefPose->SkeletalMesh, NewInstanceData);
+
+	// Push our result back
+	VM.PushValue(UE::AnimNext::KEYFRAME_STACK_NAME, MoveTemp(Keyframe));
+}
+
+void FAnimNextRigLogicTask::UpdateControlCurves(FRigInstance* InRigInstance,
+	const FBlendedCurve& Curves,
+	const FCachedIndexedCurve& ControlAttributeCurves,
+	const FCachedIndexedCurve& NeuralNetworkMaskCurves,
+	const TArrayView<const float> NeutralJointValues,
+	const TArrayView<const UE::AnimNext::FPoseBoneControlAttributeMapping> SparseDriverJointsToControlAttributesMap,
+	const TArrayView<const UE::AnimNext::FPoseBoneControlAttributeMapping> DenseDriverJointsToControlAttributesMap,
+	const UE::AnimNext::FLODPoseStack& InputPose) const
+{
+	// Combine control attribute curve with input curve to get indexed curve to apply to rig
+	// Curve elements that dont have a control mapping will have INDEX_NONE as their index
+	UE::Anim::FNamedValueArrayUtils::Union(Curves, ControlAttributeCurves,
+		[this, InRigInstance](const UE::Anim::FCurveElement& InCurveElement, const UE::Anim::FCurveElementIndexed& InControlAttributeCurveElement, UE::Anim::ENamedValueUnionFlags InFlags)
+		{
+			if (InControlAttributeCurveElement.Index != INDEX_NONE)
+			{
+				InRigInstance->SetRawControl(InControlAttributeCurveElement.Index, InCurveElement.Value);
+			}
+		});
+
+	// Combine control attribute curve with input curve to get indexed curve to apply to rig
+	// Curve elements that dont have a control mapping will have INDEX_NONE as their index
+	UE::Anim::FNamedValueArrayUtils::Union(Curves, ControlAttributeCurves,
+		[this, InRigInstance](const UE::Anim::FCurveElement& InCurveElement, const UE::Anim::FCurveElementIndexed& InControlAttributeCurveElement, UE::Anim::ENamedValueUnionFlags InFlags)
+		{
+			if (InControlAttributeCurveElement.Index != INDEX_NONE)
+			{
+				InRigInstance->SetRawControl(InControlAttributeCurveElement.Index, InCurveElement.Value);
+			}
+		});
+
+	const uint16 LODLevel = InRigInstance->GetLOD();
+	const TArrayView<const FBoneIndexType> SkeletonIndexToPoseIndex = InputPose.GetSkeletonBoneIndexToLODBoneIndexMap(); // skeleton -> pose
+
+	// The sparse mapping is NOT guaranteed to supply all quaternion attributes, so checks for each attribute mapping are present
+	for (int32 MappingIndex = 0; MappingIndex < SparseDriverJointsToControlAttributesMap.Num(); ++MappingIndex)
+	{
+		const UE::AnimNext::FPoseBoneControlAttributeMapping& Mapping = SparseDriverJointsToControlAttributesMap[MappingIndex];
+		const int32 SkeletonBoneIndex = Mapping.SkeletonBoneIndex;
+
+		// Translate from skeleton bone to LOD pose bone index as the LOD pose bone indices might vary across the skel mesh instances.
+		const int32 PoseBoneIndex = SkeletonIndexToPoseIndex[SkeletonBoneIndex];
+
+		const UE::AnimNext::FTransformSoAAdapterConst TransformAdapter = InputPose.LocalTransforms[PoseBoneIndex];
+
+		// Translation and Scale is currently not used here, so to avoid the overhead of checking them, they are simply ignored.
+		// Should the need arise to use them as well, this code will need adjustment.
+		const FQuat Rotation = TransformAdapter.GetRotation();
+		const int32 AttrIndex = Mapping.DNAJointIndex * ATTR_COUNT_PER_JOINT;
+
+		const tdm::fquat NeutralRotation{ NeutralJointValues[AttrIndex + 3], NeutralJointValues[AttrIndex + 4], NeutralJointValues[AttrIndex + 5], NeutralJointValues[AttrIndex + 6] };
+		const tdm::fquat AbsPoseRotation{ static_cast<float>(Rotation.X), static_cast<float>(Rotation.Y), static_cast<float>(Rotation.Z), static_cast<float>(Rotation.W) };
+		const tdm::fquat DeltaPoseRotation = tdm::inverse(NeutralRotation) * AbsPoseRotation;
+
+		if (Mapping.RotationX != INDEX_NONE)
+		{
+			InRigInstance->SetRawControl(Mapping.RotationX, DeltaPoseRotation.x);
+		}
+		if (Mapping.RotationY != INDEX_NONE)
+		{
+			InRigInstance->SetRawControl(Mapping.RotationY, DeltaPoseRotation.y);
+		}
+		if (Mapping.RotationZ != INDEX_NONE)
+		{
+			InRigInstance->SetRawControl(Mapping.RotationZ, DeltaPoseRotation.z);
+		}
+		if (Mapping.RotationW != INDEX_NONE)
+		{
+			InRigInstance->SetRawControl(Mapping.RotationW, DeltaPoseRotation.w);
+		}
+	}
+
+	// The dense mapping is guaranteed to supply all quaternion attributes, so NO checks for each attribute mapping are present
+	for (int32 MappingIndex = 0; MappingIndex < DenseDriverJointsToControlAttributesMap.Num(); ++MappingIndex)
+	{
+		const UE::AnimNext::FPoseBoneControlAttributeMapping& Mapping = SparseDriverJointsToControlAttributesMap[MappingIndex];
+		const int32 SkeletonBoneIndex = Mapping.SkeletonBoneIndex;
+
+		// Translate from skeleton bone to LOD pose bone index as the LOD pose bone indices might vary across the skel mesh instances.
+		const int32 PoseBoneIndex = SkeletonIndexToPoseIndex[SkeletonBoneIndex];
+
+		const UE::AnimNext::FTransformSoAAdapterConst TransformAdapter = InputPose.LocalTransforms[PoseBoneIndex];
+
+		// Translation and Scale is currently not used here, so to avoid the overhead of checking them, they are simply ignored.
+		// Should the need arise to use them as well, this code will need adjustment.
+		const FQuat Rotation = TransformAdapter.GetRotation();
+		const int32 AttrIndex = Mapping.DNAJointIndex * ATTR_COUNT_PER_JOINT;
+
+		const tdm::fquat NeutralRotation{ NeutralJointValues[AttrIndex + 3], NeutralJointValues[AttrIndex + 4], NeutralJointValues[AttrIndex + 5], NeutralJointValues[AttrIndex + 6] };
+		const tdm::fquat AbsPoseRotation{ static_cast<float>(Rotation.X), static_cast<float>(Rotation.Y), static_cast<float>(Rotation.Z), static_cast<float>(Rotation.W) };
+		const tdm::fquat DeltaPoseRotation = tdm::inverse(NeutralRotation) * AbsPoseRotation;
+
+		InRigInstance->SetRawControl(Mapping.RotationX, DeltaPoseRotation.x);
+		InRigInstance->SetRawControl(Mapping.RotationY, DeltaPoseRotation.y);
+		InRigInstance->SetRawControl(Mapping.RotationZ, DeltaPoseRotation.z);
+		InRigInstance->SetRawControl(Mapping.RotationW, DeltaPoseRotation.w);
+	}
+
+	if (InRigInstance->GetNeuralNetworkCount() != 0)
+	{
+		UE::Anim::FNamedValueArrayUtils::Union(Curves, NeuralNetworkMaskCurves,
+			[this, InRigInstance](const UE::Anim::FCurveElement& InCurveElement, const UE::Anim::FCurveElementIndexed& InControlAttributeCurveElement, UE::Anim::ENamedValueUnionFlags InFlags)
+			{
+				if (InControlAttributeCurveElement.Index != INDEX_NONE)
+				{
+					InRigInstance->SetNeuralNetworkMask(InControlAttributeCurveElement.Index, InCurveElement.Value);
+				}
+			});
+	}
+}
+
+void FAnimNextRigLogicTask::UpdateJoints(const TArrayView<const UE::AnimNext::FRigLogicBoneMapping> RigLogicIndexToMeshIndexMapping,
+	const TArrayView<const uint16> VariableJointIndices,
+	const TArrayView<const float> NeutralJointValues,
+	const TArrayView<const float> DeltaJointValues,
+	UE::AnimNext::FLODPoseStack& OutputPose) const
+{
+	const TArrayView<const FBoneIndexType> SkeletonIndexToPoseIndex = OutputPose.GetSkeletonBoneIndexToLODBoneIndexMap(); // skeleton -> pose
+
+	const float* N = NeutralJointValues.GetData();
+	const float* D = DeltaJointValues.GetData();
+	for (const UE::AnimNext::FRigLogicBoneMapping& Mapping : RigLogicIndexToMeshIndexMapping)
+	{
+		const uint16 RigLogicJointIndex = Mapping.RigLogicJointIndex;
+		const int32 SkeletonBoneIndex = Mapping.SkeletonBoneIndex;
+		const uint16 AttrIndex = RigLogicJointIndex * ATTR_COUNT_PER_JOINT;
+
+		// Translate from skeleton bone to LOD pose bone index as the LOD pose bone indices might vary across the skel mesh instances.
+		const int32 PoseBoneIndex = SkeletonIndexToPoseIndex[SkeletonBoneIndex];
+
+		UE::AnimNext::FTransformSoAAdapter TransformAdapter = OutputPose.LocalTransforms[PoseBoneIndex];
+
+		TransformAdapter.SetTranslation(FVector((N[AttrIndex + 0] + D[AttrIndex + 0]), (N[AttrIndex + 1] + D[AttrIndex + 1]), (N[AttrIndex + 2] + D[AttrIndex + 2])));
+		TransformAdapter.SetRotation(FQuat(N[AttrIndex + 3], N[AttrIndex + 4], N[AttrIndex + 5], N[AttrIndex + 6]) * FQuat(D[AttrIndex + 3], D[AttrIndex + 4], D[AttrIndex + 5], D[AttrIndex + 6]));
+		TransformAdapter.SetScale3D(FVector((N[AttrIndex + 7] + D[AttrIndex + 7]), (N[AttrIndex + 8] + D[AttrIndex + 8]), (N[AttrIndex + 9] + D[AttrIndex + 9])));
+	}
+}
+
+void FAnimNextRigLogicTask::UpdateBlendShapeCurves(const FCachedIndexedCurve& MorphTargetCurves, const TArrayView<const float> BlendShapeValues, FBlendedCurve& OutputCurves) const
+{
+	UE::Anim::FNamedValueArrayUtils::Union(OutputCurves, MorphTargetCurves, [&BlendShapeValues](UE::Anim::FCurveElement& InOutResult, const UE::Anim::FCurveElementIndexed& InSource, UE::Anim::ENamedValueUnionFlags InFlags)
+		{
+			if (BlendShapeValues.IsValidIndex(InSource.Index))
+			{
+				InOutResult.Value = BlendShapeValues[InSource.Index];
+				InOutResult.Flags |= UE::Anim::ECurveElementFlags::MorphTarget;
+			}
+		});
+}
+
+void FAnimNextRigLogicTask::UpdateAnimMapCurves(const FCachedIndexedCurve& MaskMultiplierCurves, const TArrayView<const float> AnimMapOutputs, FBlendedCurve& OutputCurves) const
+{
+	UE::Anim::FNamedValueArrayUtils::Union(OutputCurves, MaskMultiplierCurves, [&AnimMapOutputs](UE::Anim::FCurveElement& InOutResult, const UE::Anim::FCurveElementIndexed& InSource, UE::Anim::ENamedValueUnionFlags InFlags)
+		{
+			if (AnimMapOutputs.IsValidIndex(InSource.Index))
+			{
+				InOutResult.Value = AnimMapOutputs[InSource.Index];
+				InOutResult.Flags |= UE::Anim::ECurveElementFlags::Material;
+			}
+		});
+}
\ No newline at end of file
diff --git a/Engine/Plugins/Experimental/RigLogicAnimNext/Source/RigLogicAnimNext/Private/RigLogicTrait.cpp b/Engine/Plugins/Experimental/RigLogicAnimNext/Source/RigLogicAnimNext/Private/RigLogicTrait.cpp
new file mode 100644
--- /dev/null
+++ b/Engine/Plugins/Experimental/RigLogicAnimNext/Source/RigLogicAnimNext/Private/RigLogicTrait.cpp
@@ -0,0 +1,71 @@
+// Copyright Epic Games, Inc. All Rights Reserved.
+
+#include "RigLogicTrait.h"
+
+#include "RigLogicAnimNext.h"
+#include "RigLogicTask.h"
+#include "TraitCore/ExecutionContext.h"
+
+namespace UE::AnimNext
+{
+	AUTO_REGISTER_ANIM_TRAIT(FRigLogicTrait)
+
+		// Trait implementation boilerplate
+#define TRAIT_INTERFACE_ENUMERATOR(GeneratorMacro) \
+		GeneratorMacro(IEvaluate) \
+		GeneratorMacro(IHierarchy) \
+		GeneratorMacro(IUpdate) \
+		GeneratorMacro(IUpdateTraversal) \
+
+		GENERATE_ANIM_TRAIT_IMPLEMENTATION(FRigLogicTrait, TRAIT_INTERFACE_ENUMERATOR, NULL_ANIM_TRAIT_INTERFACE_ENUMERATOR, NULL_ANIM_TRAIT_EVENT_ENUMERATOR)
+#undef TRAIT_INTERFACE_ENUMERATOR
+
+	void FRigLogicTrait::OnBecomeRelevant(FUpdateTraversalContext& Context, const TTraitBinding<IUpdate>& Binding, const FTraitUpdateState& TraitState) const
+	{
+		const FSharedData* SharedData = Binding.GetSharedData<FSharedData>();
+		FInstanceData* InstanceData = Binding.GetInstanceData<FInstanceData>();
+
+		if (!InstanceData->Input.IsValid())
+		{
+			InstanceData->Input = Context.AllocateNodeInstance(Binding, SharedData->Input);
+		}
+
+		FRigLogicAnimNextModule& Module = FModuleManager::GetModuleChecked<FRigLogicAnimNextModule>("RigLogicAnimNext");
+		Module.DataPool.GarbageCollect();
+	}
+
+	void FRigLogicTrait::PostEvaluate(FEvaluateTraversalContext& Context, const TTraitBinding<IEvaluate>& Binding) const
+	{
+		FInstanceData* InstanceData = Binding.GetInstanceData<FInstanceData>();
+
+		if (InstanceData->Input.IsValid())
+		{
+			// TODO: Only add the task in case we want to actually want to apply facial animation for the current LOD.
+			//if (LODLevel < LODThreshold)
+			{
+				Context.AppendTask(FAnimNextRigLogicTask::Make());
+			}
+		}
+	}
+
+	void FRigLogicTrait::QueueChildrenForTraversal(FUpdateTraversalContext& Context, const TTraitBinding<IUpdateTraversal>& Binding, const FTraitUpdateState& TraitState, FUpdateTraversalQueue& TraversalQueue) const
+	{
+		const FInstanceData* InstanceData = Binding.GetInstanceData<FInstanceData>();
+
+		if (InstanceData->Input.IsValid())
+		{
+			TraversalQueue.Push(InstanceData->Input, TraitState);
+		}
+	}
+
+	uint32 FRigLogicTrait::GetNumChildren(const FExecutionContext& Context, const TTraitBinding<IHierarchy>& Binding) const
+	{
+		return 1;
+	}
+
+	void FRigLogicTrait::GetChildren(const FExecutionContext& Context, const TTraitBinding<IHierarchy>& Binding, FChildrenArray& Children) const
+	{
+		const FInstanceData* InstanceData = Binding.GetInstanceData<FInstanceData>();
+		Children.Add(InstanceData->Input);
+	}
+} // namespace UE::AnimNext
\ No newline at end of file
diff --git a/Engine/Plugins/Experimental/RigLogicAnimNext/Source/RigLogicAnimNext/Public/RigLogicAnimNext.h b/Engine/Plugins/Experimental/RigLogicAnimNext/Source/RigLogicAnimNext/Public/RigLogicAnimNext.h
new file mode 100644
--- /dev/null
+++ b/Engine/Plugins/Experimental/RigLogicAnimNext/Source/RigLogicAnimNext/Public/RigLogicAnimNext.h
@@ -0,0 +1,18 @@
+// Copyright Epic Games, Inc. All Rights Reserved.
+
+#pragma once
+
+#include "Modules/ModuleManager.h"
+#include "RigLogicInstanceDataPool.h"
+
+DECLARE_LOG_CATEGORY_EXTERN(LogRigLogicAnimNext, Log, All);
+
+class FRigLogicAnimNextModule : public IModuleInterface
+{
+public:
+	/** IModuleInterface implementation */
+	virtual void StartupModule() override;
+	virtual void ShutdownModule() override;
+
+	UE::AnimNext::FRigLogicInstanceDataPool DataPool;
+};
\ No newline at end of file
diff --git a/Engine/Plugins/Experimental/RigLogicAnimNext/Source/RigLogicAnimNext/Public/RigLogicInstanceData.h b/Engine/Plugins/Experimental/RigLogicAnimNext/Source/RigLogicAnimNext/Public/RigLogicInstanceData.h
new file mode 100644
--- /dev/null
+++ b/Engine/Plugins/Experimental/RigLogicAnimNext/Source/RigLogicAnimNext/Public/RigLogicInstanceData.h
@@ -0,0 +1,55 @@
+// Copyright Epic Games, Inc. All Rights Reserved.
+
+#pragma once
+
+#include "DNAIndexMapping.h"
+#include "RigInstance.h"
+#include "LODPose.h"
+#include "Engine/SkeletalMesh.h"
+
+struct FSharedRigRuntimeContext;
+struct FDNAIndexMapping;
+
+namespace UE::AnimNext
+{
+	struct FRigLogicBoneMapping
+	{
+		uint16 RigLogicJointIndex;
+		int32 SkeletonBoneIndex;
+	};
+
+	struct FPoseBoneControlAttributeMapping
+	{
+		int32 SkeletonBoneIndex;
+		int32 DNAJointIndex;
+		int32 RotationX;
+		int32 RotationY;
+		int32 RotationZ;
+		int32 RotationW;
+	};
+
+	// Instance data not per anim graph node, skeleton pair instance. TODO: Better call it PoolData?
+	struct FRigLogicAnimNextInstanceData
+	{
+		/** Cached pointer to the shared RigLogic runtime context originally owned by UDNAAsset. */
+		TSharedPtr<FSharedRigRuntimeContext> CachedRigRuntimeContext;
+
+		/** Cached pointer to the DNA index mapping which is originally owned by UDNAAsset. */
+		TSharedPtr<FDNAIndexMapping> CachedDNAIndexMapping;
+
+		/** Actually cloned RigLogic instance owned by this class. */
+		TUniquePtr<FRigInstance> RigInstance;
+
+		/** Bone index mapping from a RigLogic joint index to the reference skeleton bone index, one per LOD level. */
+		TArray<TArray<FRigLogicBoneMapping>> RigLogicToSkeletonBoneIndexMappingPerLOD;
+
+		TArray<FPoseBoneControlAttributeMapping> SparseDriverJointsToControlAttributesMap;
+		TArray<FPoseBoneControlAttributeMapping> DenseDriverJointsToControlAttributesMap;
+
+		void Init(USkeletalMesh* SkeletalMesh);
+
+	private:
+		void InitBoneIndexMapping();
+		void InitSparseAndDenseDriverJointMapping();
+	};
+} // namespace UE::AnimNext
\ No newline at end of file
diff --git a/Engine/Plugins/Experimental/RigLogicAnimNext/Source/RigLogicAnimNext/Public/RigLogicInstanceDataPool.h b/Engine/Plugins/Experimental/RigLogicAnimNext/Source/RigLogicAnimNext/Public/RigLogicInstanceDataPool.h
new file mode 100644
--- /dev/null
+++ b/Engine/Plugins/Experimental/RigLogicAnimNext/Source/RigLogicAnimNext/Public/RigLogicInstanceDataPool.h
@@ -0,0 +1,22 @@
+// Copyright Epic Games, Inc. All Rights Reserved.
+
+#pragma once
+
+#include "RigLogicInstanceData.h"
+
+namespace UE::AnimNext
+{
+	class FRigLogicInstanceDataPool
+	{
+	public:
+		TMap<TWeakObjectPtr<const USkeletalMesh>, TArray<TSharedPtr<FRigLogicAnimNextInstanceData>>> Datas;
+
+		mutable FCriticalSection PoolAccessCriticalSection;
+
+		TSharedPtr<FRigLogicAnimNextInstanceData> RequestData(const UE::AnimNext::FReferencePose* InReferencePose);
+		void FreeData(TWeakObjectPtr<const USkeletalMesh> SkeletalMesh, TSharedPtr<FRigLogicAnimNextInstanceData> InData);
+
+		void GarbageCollect();
+		void Log();
+	};
+} // namespace UE::AnimNext
\ No newline at end of file
diff --git a/Engine/Plugins/Experimental/RigLogicAnimNext/Source/RigLogicAnimNext/Public/RigLogicTask.h b/Engine/Plugins/Experimental/RigLogicAnimNext/Source/RigLogicAnimNext/Public/RigLogicTask.h
new file mode 100644
--- /dev/null
+++ b/Engine/Plugins/Experimental/RigLogicAnimNext/Source/RigLogicAnimNext/Public/RigLogicTask.h
@@ -0,0 +1,55 @@
+// Copyright Epic Games, Inc. All Rights Reserved.
+
+#pragma once
+
+#include "CoreMinimal.h"
+#include "EvaluationVM/EvaluationTask.h"
+#include "Animation/AnimCurveTypes.h"
+#include "LODPose.h"
+#include "AnimNextStats.h"
+
+#include "RigInstance.h"
+#include "RigLogic.h"
+#include "RigLogicTrait.h"
+#include "RigLogicInstanceData.h"
+
+#include "RigLogicTask.generated.h"
+
+struct FSharedRigRuntimeContext;
+struct FDNAIndexMapping;
+
+DECLARE_CYCLE_STAT_EXTERN(TEXT("AnimNext Task: RigLogic"), STAT_AnimNext_Task_RigLogic, STATGROUP_AnimNext, RIGLOGICANIMNEXT_API);
+
+USTRUCT()
+struct RIGLOGICANIMNEXT_API FAnimNextRigLogicTask : public FAnimNextEvaluationTask
+{
+	GENERATED_BODY()
+
+	DECLARE_ANIM_EVALUATION_TASK(FAnimNextRigLogicTask)
+
+	static FAnimNextRigLogicTask Make();
+
+	// Task entry point
+	virtual void Execute(UE::AnimNext::FEvaluationVM& VM) const override;
+
+	using FCachedIndexedCurve = TBaseBlendedCurve<FDefaultAllocator, UE::Anim::FCurveElementIndexed>;
+	static constexpr uint16 ATTR_COUNT_PER_JOINT = 10;
+
+	void UpdateControlCurves(FRigInstance* InRigInstance,
+		const FBlendedCurve& Curves,
+		const FCachedIndexedCurve& ControlAttributeCurves,
+		const FCachedIndexedCurve& NeuralNetworkMaskCurves,
+		const TArrayView<const float> NeutralJointValues,
+		const TArrayView<const UE::AnimNext::FPoseBoneControlAttributeMapping> SparseDriverJointsToControlAttributesMap,
+		const TArrayView<const UE::AnimNext::FPoseBoneControlAttributeMapping> DenseDriverJointsToControlAttributesMap,
+		const UE::AnimNext::FLODPoseStack& InputPose) const;
+
+	void UpdateJoints(const TArrayView<const UE::AnimNext::FRigLogicBoneMapping> RigLogicIndexToMeshIndexMapping,
+		const TArrayView<const uint16> VariableJointIndices,
+		const TArrayView<const float> NeutralJointValues,
+		const TArrayView<const float> DeltaJointValues,
+		UE::AnimNext::FLODPoseStack& OutputPose) const;
+
+	void UpdateBlendShapeCurves(const FCachedIndexedCurve& MorphTargetCurves, const TArrayView<const float> BlendShapeValues, FBlendedCurve& OutputCurves) const;
+	void UpdateAnimMapCurves(const FCachedIndexedCurve& MaskMultiplierCurves, const TArrayView<const float> AnimMapOutputs, FBlendedCurve& OutputCurves) const;
+};
\ No newline at end of file
diff --git a/Engine/Plugins/Experimental/RigLogicAnimNext/Source/RigLogicAnimNext/Public/RigLogicTrait.h b/Engine/Plugins/Experimental/RigLogicAnimNext/Source/RigLogicAnimNext/Public/RigLogicTrait.h
new file mode 100644
--- /dev/null
+++ b/Engine/Plugins/Experimental/RigLogicAnimNext/Source/RigLogicAnimNext/Public/RigLogicTrait.h
@@ -0,0 +1,73 @@
+// Copyright Epic Games, Inc. All Rights Reserved.
+
+#pragma once
+
+#include "CoreMinimal.h"
+
+#include "TraitCore/Trait.h"
+#include "TraitInterfaces/IContinuousBlend.h"
+#include "TraitInterfaces/IEvaluate.h"
+#include "TraitInterfaces/IHierarchy.h"
+#include "TraitInterfaces/IUpdate.h"
+
+#include "RigLogicTrait.generated.h"
+
+USTRUCT(meta = (DisplayName = "RigLogic"))
+struct RIGLOGICANIMNEXT_API FAnimNextRigLogicTraitSharedData : public FAnimNextTraitSharedData
+{
+	GENERATED_BODY()
+
+	UPROPERTY()
+	FAnimNextTraitHandle Input;
+
+	/*
+	 * Max LOD level that post-process AnimBPs are evaluated.
+	 * For example if you have the threshold set to 2, it will evaluate until including LOD 2 (based on 0 index). In case the LOD level gets set to 3, it will stop evaluating the post-process AnimBP.
+	 * Setting it to -1 will always evaluate it and disable LODing.
+	 */
+	UPROPERTY(EditAnywhere, Category = RigLogic)
+	int32 LODThreshold = INDEX_NONE;
+
+	// Latent pin support boilerplate
+#define TRAIT_LATENT_PROPERTIES_ENUMERATOR(GeneratorMacro) \
+		GeneratorMacro(LODThreshold)
+
+	GENERATE_TRAIT_LATENT_PROPERTIES(FAnimNextRigLogicTraitSharedData, TRAIT_LATENT_PROPERTIES_ENUMERATOR)
+#undef TRAIT_LATENT_PROPERTIES_ENUMERATOR
+};
+
+namespace UE::AnimNext
+{
+	/**
+	 * FRigLogicTrait
+	 * 
+	 * A trait that can run RigLogic.
+	 */
+	struct RIGLOGICANIMNEXT_API FRigLogicTrait : FBaseTrait, IEvaluate, IUpdate, IUpdateTraversal, IHierarchy
+	{
+		DECLARE_ANIM_TRAIT(FRigLogicTrait, 0x4311e25e, FBaseTrait)
+
+		using FSharedData = FAnimNextRigLogicTraitSharedData;
+
+		// Do we need this? Or is the input in the shared data sufficient?
+		struct FInstanceData : FTrait::FInstanceData
+		{
+			/** Input node from which we receive the input pose as well as the facial expression curves. */
+			FTraitPtr Input;
+		};
+
+		// IEvaluate impl
+		virtual void PostEvaluate(FEvaluateTraversalContext& Context, const TTraitBinding<IEvaluate>& Binding) const override;
+
+		// IUpdate impl
+		// Called before the first update when a trait stack becomes relevant
+		virtual void OnBecomeRelevant(FUpdateTraversalContext& Context, const TTraitBinding<IUpdate>& Binding, const FTraitUpdateState& TraitState) const override;
+
+		// IUpdateTraversal impl
+		virtual void QueueChildrenForTraversal(FUpdateTraversalContext& Context, const TTraitBinding<IUpdateTraversal>& Binding, const FTraitUpdateState& TraitState, FUpdateTraversalQueue& TraversalQueue) const override;
+
+		// IHierarchy impl
+		virtual uint32 GetNumChildren(const FExecutionContext& Context, const TTraitBinding<IHierarchy>& Binding) const override;
+		virtual void GetChildren(const FExecutionContext& Context, const TTraitBinding<IHierarchy>& Binding, FChildrenArray& Children) const override;
+	};
+} // namespace UE::AnimNext
\ No newline at end of file
diff --git a/Engine/Plugins/Experimental/RigLogicAnimNext/Source/RigLogicAnimNext/RigLogicAnimNext.Build.cs b/Engine/Plugins/Experimental/RigLogicAnimNext/Source/RigLogicAnimNext/RigLogicAnimNext.Build.cs
new file mode 100644
--- /dev/null
+++ b/Engine/Plugins/Experimental/RigLogicAnimNext/Source/RigLogicAnimNext/RigLogicAnimNext.Build.cs
@@ -0,0 +1,54 @@
+// Copyright Epic Games, Inc. All Rights Reserved.
+
+using UnrealBuildTool;
+
+public class RigLogicAnimNext : ModuleRules
+{
+	public RigLogicAnimNext(ReadOnlyTargetRules Target) : base(Target)
+	{
+		PCHUsage = ModuleRules.PCHUsageMode.UseExplicitOrSharedPCHs;
+		
+		PublicIncludePaths.AddRange(
+			new string[]
+			{
+			}
+			);
+				
+		
+		PrivateIncludePaths.AddRange(
+			new string[]
+			{
+			}
+			);
+			
+		
+		PublicDependencyModuleNames.AddRange(
+			new string[]
+			{
+				"Core",
+				"Engine",
+				"RigLogicModule",
+				"AnimNext",
+				"AnimNextAnimGraph",
+			}
+			);
+			
+		
+		PrivateDependencyModuleNames.AddRange(
+			new string[]
+			{
+				"CoreUObject",
+				"Engine",
+				"Slate",
+				"SlateCore",
+			}
+			);
+		
+		
+		DynamicallyLoadedModuleNames.AddRange(
+			new string[]
+			{
+			}
+			);
+	}
+}
