RigLogic Body Correctives for AnimNext

This change enables the new body correctives using RigLogic that will replace MH pose drivers in AnimNext.

* AnimNext's reference pose now generates and stores a mapping from mesh bone indices to LOD pose bone indices.
* We're now using pose indices directly without any bone index translation at runtime. This is more cache friendly as we're less jumping around in the transform arrays and should help increase performance.
* Fixed a bone index translation issue when initializing the bone index mapping.
* Dense and sparse driver joint mapping is now stored per LOD and only the needed data is sent over to the trait tasks.
* Some more safety checks and warning/error logging.

diff --git a/Engine/Plugins/Experimental/AnimNext/Source/AnimNext/Private/GenerationTools.cpp b/Engine/Plugins/Experimental/AnimNext/Source/AnimNext/Private/GenerationTools.cpp
--- a/Engine/Plugins/Experimental/AnimNext/Source/AnimNext/Private/GenerationTools.cpp
+++ b/Engine/Plugins/Experimental/AnimNext/Source/AnimNext/Private/GenerationTools.cpp
@@ -131,8 +131,10 @@ struct FCompareBoneIndexType
 				const int32 NumOrderedBones = LODBoneIndexToMeshBoneIndexMap.Num();
 
 				TArray<FBoneIndexType> SkeletonBoneIndexToLODBoneIndexMap;
+				TArray<FBoneIndexType> MeshBoneIndexToLODBoneIndexMap;
 				TArray<FBoneIndexType> LODBoneIndexToSkeletonBoneIndexMap;
 				SkeletonBoneIndexToLODBoneIndexMap.Init(INDEX_NONE, NumMeshBones);
+				MeshBoneIndexToLODBoneIndexMap.Init(INDEX_NONE, NumMeshBones);
 				LODBoneIndexToSkeletonBoneIndexMap.Init(INDEX_NONE, NumOrderedBones);
 
 				for (int32 LODBoneIndex = 0; LODBoneIndex < NumOrderedBones; ++LODBoneIndex)
@@ -145,6 +147,7 @@ struct FCompareBoneIndexType
 					ensure(SkeletonBoneIndex.IsValid());	// We expect the skeletal mesh bone to map to a valid skeleton bone
 
 					SkeletonBoneIndexToLODBoneIndexMap[static_cast<FBoneIndexType>(SkeletonBoneIndex.GetInt())] = LODBoneIndex;
+					MeshBoneIndexToLODBoneIndexMap[static_cast<FBoneIndexType>(MeshBoneIndex.GetInt())] = LODBoneIndex;
 					LODBoneIndexToSkeletonBoneIndexMap[LODBoneIndex] = static_cast<FBoneIndexType>(SkeletonBoneIndex.GetInt());
 				}
 
@@ -177,6 +180,7 @@ struct FCompareBoneIndexType
 					{ LODBoneIndexToMeshBoneIndexMap },
 					{ LODBoneIndexToSkeletonBoneIndexMap },
 					{ SkeletonBoneIndexToLODBoneIndexMap },
+					{ MeshBoneIndexToLODBoneIndexMap },
 					LODNumBones,
 					bCanGenerateSingleBonesList);
 
diff --git a/Engine/Plugins/Experimental/AnimNext/Source/AnimNext/Public/ReferencePose.h b/Engine/Plugins/Experimental/AnimNext/Source/AnimNext/Public/ReferencePose.h
--- a/Engine/Plugins/Experimental/AnimNext/Source/AnimNext/Public/ReferencePose.h
+++ b/Engine/Plugins/Experimental/AnimNext/Source/AnimNext/Public/ReferencePose.h
@@ -54,6 +54,11 @@ struct TReferencePose
 	// When fast path is enabled, we have a single LOD entry
 	TArray<TArray<FBoneIndexType, AllocatorType>, AllocatorType> SkeletonBoneIndexToLODBoneIndexMapPerLOD;
 
+	// List of pose bone indices by mesh bone index for each LOD
+	// Each list of bone indices is a mapping of: SkeletalMeshBoneIndex -> LODSortedBoneIndex
+	// When fast path is enabled, we have a single LOD entry
+	TArray<TArray<FBoneIndexType, AllocatorType>, AllocatorType> MeshBoneIndexToLODBoneIndexMapPerLOD;
+
 	// Number of bones for each LOD
 	TArray<int32, AllocatorType> LODNumBones;
 
@@ -88,6 +93,7 @@ struct TReferencePose
 		, const TArray<TArray<FBoneIndexType>>& InLODBoneIndexToMeshBoneIndexMapPerLOD
 		, const TArray<TArray<FBoneIndexType>>& InLODBoneIndexToSkeletonBoneIndexMapPerLOD
 		, const TArray<TArray<FBoneIndexType>>& InSkeletonBoneIndexToLODBoneIndexMapPerLOD
+		, const TArray<TArray<FBoneIndexType>>& InMeshBoneIndexToLODBoneIndexMapPerLOD
 		, const TArray<int32, AllocatorType>& InLODNumBones
 		, bool bFastPath = false)
 	{
@@ -100,6 +106,7 @@ struct TReferencePose
 		LODBoneIndexToMeshBoneIndexMapPerLOD = InLODBoneIndexToMeshBoneIndexMapPerLOD;
 		LODBoneIndexToSkeletonBoneIndexMapPerLOD = InLODBoneIndexToSkeletonBoneIndexMapPerLOD;
 		SkeletonBoneIndexToLODBoneIndexMapPerLOD = InSkeletonBoneIndexToLODBoneIndexMapPerLOD;
+		MeshBoneIndexToLODBoneIndexMapPerLOD = InMeshBoneIndexToLODBoneIndexMapPerLOD;
 		LODNumBones = InLODNumBones;
 
 		const TArray<FTransform>& RefBonePose = RefSkeleton.GetRefBonePose();
@@ -182,6 +189,20 @@ struct TReferencePose
 		return ArrayView;
 	}
 
+	// Return a list of LOD bone indices, a mapping of: SkeletalMeshBoneIndex -> LODSortedBoneIndex
+	const TArrayView<const FBoneIndexType> GetMeshBoneIndexToLODBoneIndexMap(int32 LODLevel) const
+	{
+		TArrayView<const FBoneIndexType> ArrayView;
+
+		if (LODLevel >= 0 && (IsFastPath() || LODLevel < MeshBoneIndexToLODBoneIndexMapPerLOD.Num()))
+		{
+			const int32 LODIndex = IsFastPath() ? 0 : LODLevel;
+			ArrayView = MeshBoneIndexToLODBoneIndexMapPerLOD[LODIndex];
+		}
+
+		return ArrayView;
+	}
+
 	int32 GetMeshBoneIndexFromLODBoneIndex(int32 LODBoneIndex) const
 	{
 		const int32 NumBonesLOD0 = LODBoneIndexToMeshBoneIndexMapPerLOD[0].Num();
@@ -202,6 +223,14 @@ struct TReferencePose
 		return LODBoneIndexToSkeletonBoneIndexMapPerLOD[0][LODBoneIndex];
 	}
 
+	// Translate a skeletal mesh bone index to a pose bone index.
+	int32 GetLODBoneIndexFromMeshBoneIndex(int32 MeshBoneIndex) const
+	{
+		const int32 NumBonesLOD0 = MeshBoneIndexToLODBoneIndexMapPerLOD[0].Num();
+		check(MeshBoneIndex < NumBonesLOD0);
+		return MeshBoneIndexToLODBoneIndexMapPerLOD[0][MeshBoneIndex];
+	}
+
 	int32 GetLODBoneIndexFromSkeletonBoneIndex(int32 SkeletionBoneIndex) const
 	{
 		const int32 NumBonesLOD0 = SkeletonBoneIndexToLODBoneIndexMapPerLOD[0].Num();
diff --git a/Engine/Plugins/Experimental/RigLogicAnimNext/Source/RigLogicAnimNext/Private/RigLogicInstanceData.cpp b/Engine/Plugins/Experimental/RigLogicAnimNext/Source/RigLogicAnimNext/Private/RigLogicInstanceData.cpp
--- a/Engine/Plugins/Experimental/RigLogicAnimNext/Source/RigLogicAnimNext/Private/RigLogicInstanceData.cpp
+++ b/Engine/Plugins/Experimental/RigLogicAnimNext/Source/RigLogicAnimNext/Private/RigLogicInstanceData.cpp
@@ -3,6 +3,7 @@
 #include "RigLogicInstanceData.h"
 #include "RigLogicAnimNext.h"
 
+#include "Animation/Skeleton.h"
 #include "Engine/SkeletalMesh.h"
 #include "DNAAsset.h"
 #include "LODPose.h"
@@ -10,12 +11,25 @@
 
 namespace UE::AnimNext
 {
-	void FRigLogicAnimNextInstanceData::Init(USkeletalMesh* SkeletalMesh)
+	void FRigLogicAnimNextInstanceData::Init(const UE::AnimNext::FReferencePose* ReferencePose)
 	{
+		if (!ReferencePose)
+		{
+			UE_LOG(LogRigLogicAnimNext, Error, TEXT("Reference pose invalid."));
+			return;
+		}
+
+		USkeletalMesh* SkeletalMesh = const_cast<USkeletalMesh*>(ReferencePose->SkeletalMesh.Get());
+		if (!SkeletalMesh)
+		{
+			UE_LOG(LogRigLogicAnimNext, Error, TEXT("No skeletal mesh assigned to reference pose."));
+			return;
+		}
+
 		const USkeleton* Skeleton = SkeletalMesh->GetSkeleton();
 		if (!Skeleton)
 		{
-			UE_LOG(LogRigLogicAnimNext, Warning, TEXT("No skeleton assigned to the skeletal mesh."));
+			UE_LOG(LogRigLogicAnimNext, Error, TEXT("No skeleton assigned to the skeletal mesh."));
 			return;
 		}
 
@@ -29,7 +43,7 @@ namespace UE::AnimNext
 		TSharedPtr<FSharedRigRuntimeContext> SharedRigRuntimeContext = DNAAsset->GetRigRuntimeContext();
 		if (!SharedRigRuntimeContext.IsValid())
 		{
-			UE_LOG(LogRigLogicAnimNext, Warning, TEXT("Can't get the shared rig runtime context."));
+			UE_LOG(LogRigLogicAnimNext, Error, TEXT("Can't get the shared rig runtime context."));
 			return;
 		}
 
@@ -41,11 +55,11 @@ namespace UE::AnimNext
 
 		CachedDNAIndexMapping = DNAAsset->GetDNAIndexMapping(Skeleton, SkeletalMesh);
 
-		InitBoneIndexMapping();
-		InitSparseAndDenseDriverJointMapping();
+		InitBoneIndexMapping(ReferencePose);
+		InitSparseAndDenseDriverJointMapping(ReferencePose);
 	}
 
-	void FRigLogicAnimNextInstanceData::InitBoneIndexMapping()
+	void FRigLogicAnimNextInstanceData::InitBoneIndexMapping(const UE::AnimNext::FReferencePose* ReferencePose)
 	{
 		const uint32 NumLODs = CachedRigRuntimeContext->VariableJointIndicesPerLOD.Num();
 		RigLogicToSkeletonBoneIndexMappingPerLOD.Empty();
@@ -56,13 +70,18 @@ namespace UE::AnimNext
 			const TArray<uint16>& VariableJointIndices = CachedRigRuntimeContext->VariableJointIndicesPerLOD[LODLevel].Values;
 			RigLogicToSkeletonBoneIndexMappingPerLOD[LODLevel].Reserve(VariableJointIndices.Num());
 
+			const TArrayView<const FBoneIndexType> MeshToPoseBoneIndexMap = ReferencePose->GetMeshBoneIndexToLODBoneIndexMap(LODLevel); // mesh -> pose
+
 			for (const uint16 RigLogicJointIndex : VariableJointIndices)
 			{
-				const FMeshPoseBoneIndex MeshPoseBoneIndex = CachedDNAIndexMapping->JointsMapDNAIndicesToMeshPoseBoneIndices[RigLogicJointIndex];
-				if (MeshPoseBoneIndex != INDEX_NONE)
+				// Get the mesh bone index from the RigLogic joint index. Index originating from a FindBoneName() on the mesh's reference skeleton.
+				const FMeshPoseBoneIndex MeshBoneIndex = CachedDNAIndexMapping->JointsMapDNAIndicesToMeshPoseBoneIndices[RigLogicJointIndex];
+				if (MeshBoneIndex.IsValid())
 				{
-					const int32 SkeletonBoneIndex = MeshPoseBoneIndex.GetInt();
-					RigLogicToSkeletonBoneIndexMappingPerLOD[LODLevel].Add({ RigLogicJointIndex, SkeletonBoneIndex });
+					// Convert the skeleton bone index to a pose bone index for the given LOD level.
+					const int32 PoseBoneIndex = MeshToPoseBoneIndexMap[MeshBoneIndex.GetInt()];
+
+					RigLogicToSkeletonBoneIndexMappingPerLOD[LODLevel].Add({ RigLogicJointIndex, PoseBoneIndex });
 				}
 				else
 				{
@@ -72,28 +91,38 @@ namespace UE::AnimNext
 		}
 	}
 
-	void FRigLogicAnimNextInstanceData::InitSparseAndDenseDriverJointMapping()
+	void FRigLogicAnimNextInstanceData::InitSparseAndDenseDriverJointMapping(const UE::AnimNext::FReferencePose* ReferencePose)
 	{
+		const uint32 NumLODs = CachedRigRuntimeContext->VariableJointIndicesPerLOD.Num();
+
 		// Populate driver joint to raw control attribute mapping (used to feed RigLogic with inputs from the joint hierarchy)
-		SparseDriverJointsToControlAttributesMap.Empty();
-		DenseDriverJointsToControlAttributesMap.Empty();
-		DenseDriverJointsToControlAttributesMap.Reserve(CachedDNAIndexMapping->DriverJointsToControlAttributesMap.Num());
+		SparseDriverJointsToControlAttributesMapPerLOD.Empty();
+		DenseDriverJointsToControlAttributesMapPerLOD.Empty();
 
-		// Sparse mapping will likely remain empty so no reservation happens
-		for (const auto& Mapping : CachedDNAIndexMapping->DriverJointsToControlAttributesMap)
+		SparseDriverJointsToControlAttributesMapPerLOD.SetNum(NumLODs);
+		DenseDriverJointsToControlAttributesMapPerLOD.SetNum(NumLODs);
+
+		for (uint32 LODLevel = 0; LODLevel < NumLODs; ++LODLevel)
 		{
-			const FMeshPoseBoneIndex MeshPoseBoneIndex = Mapping.MeshPoseBoneIndex;
+			const TArrayView<const FBoneIndexType> MeshToPoseBoneIndexMap = ReferencePose->GetMeshBoneIndexToLODBoneIndexMap(LODLevel); // mesh -> pose
 
-			if (MeshPoseBoneIndex != INDEX_NONE)
+			// Sparse mapping will likely remain empty so no reservation happens
+			for (const auto& Mapping : CachedDNAIndexMapping->DriverJointsToControlAttributesMap)
 			{
-				const int32 SkeletonBoneIndex = MeshPoseBoneIndex.GetInt();
-				if ((Mapping.RotationX != INDEX_NONE) && (Mapping.RotationY != INDEX_NONE) && (Mapping.RotationZ != INDEX_NONE) && (Mapping.RotationW != INDEX_NONE))
+				const FMeshPoseBoneIndex MeshBoneIndex = Mapping.MeshPoseBoneIndex;
+				if (MeshBoneIndex.IsValid())
 				{
-					DenseDriverJointsToControlAttributesMap.Add({ SkeletonBoneIndex, Mapping.DNAJointIndex, Mapping.RotationX, Mapping.RotationY, Mapping.RotationZ, Mapping.RotationW });
-				}
-				else
-				{
-					SparseDriverJointsToControlAttributesMap.Add({ SkeletonBoneIndex, Mapping.DNAJointIndex, Mapping.RotationX, Mapping.RotationY, Mapping.RotationZ, Mapping.RotationW });
+					// Convert the mesh bone index to a pose bone index for the given LOD level.
+					const int32 PoseBoneIndex = MeshToPoseBoneIndexMap[MeshBoneIndex.GetInt()];
+
+					if ((Mapping.RotationX != INDEX_NONE) && (Mapping.RotationY != INDEX_NONE) && (Mapping.RotationZ != INDEX_NONE) && (Mapping.RotationW != INDEX_NONE))
+					{
+						DenseDriverJointsToControlAttributesMapPerLOD[LODLevel].Add({ PoseBoneIndex, Mapping.DNAJointIndex, Mapping.RotationX, Mapping.RotationY, Mapping.RotationZ, Mapping.RotationW });
+					}
+					else
+					{
+						SparseDriverJointsToControlAttributesMapPerLOD[LODLevel].Add({ PoseBoneIndex, Mapping.DNAJointIndex, Mapping.RotationX, Mapping.RotationY, Mapping.RotationZ, Mapping.RotationW });
+					}
 				}
 			}
 		}
diff --git a/Engine/Plugins/Experimental/RigLogicAnimNext/Source/RigLogicAnimNext/Private/RigLogicInstanceDataPool.cpp b/Engine/Plugins/Experimental/RigLogicAnimNext/Source/RigLogicAnimNext/Private/RigLogicInstanceDataPool.cpp
--- a/Engine/Plugins/Experimental/RigLogicAnimNext/Source/RigLogicAnimNext/Private/RigLogicInstanceDataPool.cpp
+++ b/Engine/Plugins/Experimental/RigLogicAnimNext/Source/RigLogicAnimNext/Private/RigLogicInstanceDataPool.cpp
@@ -26,7 +26,7 @@ namespace UE::AnimNext
 
 		// We don't have any pool datas in the array yet. Create a new one.
 		TSharedPtr<FRigLogicAnimNextInstanceData> NewObject = TSharedPtr<FRigLogicAnimNextInstanceData>(new FRigLogicAnimNextInstanceData());
-		NewObject->Init(const_cast<USkeletalMesh*>(InReferencePose->SkeletalMesh.Get()));
+		NewObject->Init(InReferencePose);
 		return NewObject;
 	}
 
diff --git a/Engine/Plugins/Experimental/RigLogicAnimNext/Source/RigLogicAnimNext/Private/RigLogicTask.cpp b/Engine/Plugins/Experimental/RigLogicAnimNext/Source/RigLogicAnimNext/Private/RigLogicTask.cpp
--- a/Engine/Plugins/Experimental/RigLogicAnimNext/Source/RigLogicAnimNext/Private/RigLogicTask.cpp
+++ b/Engine/Plugins/Experimental/RigLogicAnimNext/Source/RigLogicAnimNext/Private/RigLogicTask.cpp
@@ -45,6 +45,8 @@ void FAnimNextRigLogicTask::Execute(UE::AnimNext::FEvaluationVM& VM) const
 	{
 		if (const FRigLogic* RigLogic = RigRuntimeContext->RigLogic.Get())
 		{
+			const int32 LODLevel = Keyframe->Pose.LODLevel;
+
 			// 1. Feed input curves to RigLogic.
 			FBlendedCurve& Curves = Keyframe->Curves;
 			const TArrayView<const float> NeutralJointValues = NewInstanceData->CachedRigRuntimeContext->RigLogic->GetNeutralJointValues();
@@ -54,12 +56,11 @@ void FAnimNextRigLogicTask::Execute(UE::AnimNext::FEvaluationVM& VM) const
 				DNAIndexMapping->ControlAttributeCurves,
 				DNAIndexMapping->NeuralNetworkMaskCurves,
 				NeutralJointValues,
-				NewInstanceData->SparseDriverJointsToControlAttributesMap,
-				NewInstanceData->DenseDriverJointsToControlAttributesMap,
+				NewInstanceData->SparseDriverJointsToControlAttributesMapPerLOD[LODLevel],
+				NewInstanceData->DenseDriverJointsToControlAttributesMapPerLOD[LODLevel],
 				Keyframe->Pose);
 
 			// 2. Evaluate RigLogic.
-			const int32 LODLevel = Keyframe->Pose.LODLevel;
 			RigInstance->SetLOD(LODLevel);
 			RigLogic->Calculate(RigInstance);
 
@@ -121,10 +122,7 @@ void FAnimNextRigLogicTask::UpdateControlCurves(FRigInstance* InRigInstance,
 	for (int32 MappingIndex = 0; MappingIndex < SparseDriverJointsToControlAttributesMap.Num(); ++MappingIndex)
 	{
 		const UE::AnimNext::FPoseBoneControlAttributeMapping& Mapping = SparseDriverJointsToControlAttributesMap[MappingIndex];
-		const int32 SkeletonBoneIndex = Mapping.SkeletonBoneIndex;
-
-		// Translate from skeleton bone to LOD pose bone index as the LOD pose bone indices might vary across the skel mesh instances.
-		const int32 PoseBoneIndex = SkeletonIndexToPoseIndex[SkeletonBoneIndex];
+		const int32 PoseBoneIndex = Mapping.PoseBoneIndex;
 
 		const UE::AnimNext::FTransformSoAAdapterConst TransformAdapter = InputPose.LocalTransforms[PoseBoneIndex];
 
@@ -158,11 +156,8 @@ void FAnimNextRigLogicTask::UpdateControlCurves(FRigInstance* InRigInstance,
 	// The dense mapping is guaranteed to supply all quaternion attributes, so NO checks for each attribute mapping are present
 	for (int32 MappingIndex = 0; MappingIndex < DenseDriverJointsToControlAttributesMap.Num(); ++MappingIndex)
 	{
-		const UE::AnimNext::FPoseBoneControlAttributeMapping& Mapping = SparseDriverJointsToControlAttributesMap[MappingIndex];
-		const int32 SkeletonBoneIndex = Mapping.SkeletonBoneIndex;
-
-		// Translate from skeleton bone to LOD pose bone index as the LOD pose bone indices might vary across the skel mesh instances.
-		const int32 PoseBoneIndex = SkeletonIndexToPoseIndex[SkeletonBoneIndex];
+		const UE::AnimNext::FPoseBoneControlAttributeMapping& Mapping = DenseDriverJointsToControlAttributesMap[MappingIndex];
+		const int32 PoseBoneIndex = Mapping.PoseBoneIndex;
 
 		const UE::AnimNext::FTransformSoAAdapterConst TransformAdapter = InputPose.LocalTransforms[PoseBoneIndex];
 
@@ -204,15 +199,13 @@ void FAnimNextRigLogicTask::UpdateJoints(const TArrayView<const UE::AnimNext::FR
 
 	const float* N = NeutralJointValues.GetData();
 	const float* D = DeltaJointValues.GetData();
+
 	for (const UE::AnimNext::FRigLogicBoneMapping& Mapping : RigLogicIndexToMeshIndexMapping)
 	{
 		const uint16 RigLogicJointIndex = Mapping.RigLogicJointIndex;
-		const int32 SkeletonBoneIndex = Mapping.SkeletonBoneIndex;
+		const int32 PoseBoneIndex = Mapping.PoseBoneIndex;
 		const uint16 AttrIndex = RigLogicJointIndex * ATTR_COUNT_PER_JOINT;
 
-		// Translate from skeleton bone to LOD pose bone index as the LOD pose bone indices might vary across the skel mesh instances.
-		const int32 PoseBoneIndex = SkeletonIndexToPoseIndex[SkeletonBoneIndex];
-
 		UE::AnimNext::FTransformSoAAdapter TransformAdapter = OutputPose.LocalTransforms[PoseBoneIndex];
 
 		TransformAdapter.SetTranslation(FVector((N[AttrIndex + 0] + D[AttrIndex + 0]), (N[AttrIndex + 1] + D[AttrIndex + 1]), (N[AttrIndex + 2] + D[AttrIndex + 2])));
diff --git a/Engine/Plugins/Experimental/RigLogicAnimNext/Source/RigLogicAnimNext/Public/RigLogicInstanceData.h b/Engine/Plugins/Experimental/RigLogicAnimNext/Source/RigLogicAnimNext/Public/RigLogicInstanceData.h
--- a/Engine/Plugins/Experimental/RigLogicAnimNext/Source/RigLogicAnimNext/Public/RigLogicInstanceData.h
+++ b/Engine/Plugins/Experimental/RigLogicAnimNext/Source/RigLogicAnimNext/Public/RigLogicInstanceData.h
@@ -14,13 +14,13 @@ namespace UE::AnimNext
 {
 	struct FRigLogicBoneMapping
 	{
-		uint16 RigLogicJointIndex;
-		int32 SkeletonBoneIndex;
+		uint16 RigLogicJointIndex; /** RigLogic joint index based on the internal RigLogic skeleton. Bone indices differ from the skeleton or skeletal mesh bone indices and need translation. */
+		int32 PoseBoneIndex; /** Pose bone index based on the given LOD level. */
 	};
 
 	struct FPoseBoneControlAttributeMapping
 	{
-		int32 SkeletonBoneIndex;
+		int32 PoseBoneIndex; /** Pose bone index based on the given LOD level. */
 		int32 DNAJointIndex;
 		int32 RotationX;
 		int32 RotationY;
@@ -28,7 +28,7 @@ namespace UE::AnimNext
 		int32 RotationW;
 	};
 
-	// Instance data not per anim graph node, skeleton pair instance. TODO: Better call it PoolData?
+	// Instance data unique per reference pose, cloned by the maximum number of parallel evaluations. TODO: Better call it PoolData?
 	struct FRigLogicAnimNextInstanceData
 	{
 		/** Cached pointer to the shared RigLogic runtime context originally owned by UDNAAsset. */
@@ -43,13 +43,13 @@ namespace UE::AnimNext
 		/** Bone index mapping from a RigLogic joint index to the reference skeleton bone index, one per LOD level. */
 		TArray<TArray<FRigLogicBoneMapping>> RigLogicToSkeletonBoneIndexMappingPerLOD;
 
-		TArray<FPoseBoneControlAttributeMapping> SparseDriverJointsToControlAttributesMap;
-		TArray<FPoseBoneControlAttributeMapping> DenseDriverJointsToControlAttributesMap;
+		TArray<TArray<FPoseBoneControlAttributeMapping>> SparseDriverJointsToControlAttributesMapPerLOD;
+		TArray<TArray<FPoseBoneControlAttributeMapping>> DenseDriverJointsToControlAttributesMapPerLOD;
 
-		void Init(USkeletalMesh* SkeletalMesh);
+		void Init(const UE::AnimNext::FReferencePose* ReferencePose);
 
 	private:
-		void InitBoneIndexMapping();
-		void InitSparseAndDenseDriverJointMapping();
+		void InitBoneIndexMapping(const UE::AnimNext::FReferencePose* ReferencePose);
+		void InitSparseAndDenseDriverJointMapping(const UE::AnimNext::FReferencePose* ReferencePose);
 	};
 } // namespace UE::AnimNext
\ No newline at end of file
