Deformed faces multi-threading issue in AnimNext RigLogic

Several expressions seem to be incorrectly controlled sometimes with a high thread utilization. This is mostly only happening with limiting the cores to 5 or something and running UE outside of VS.

* Moved the rig instance to the trait instance data, so that we now have a rig instance per character.
* Cloned the DNA index mapping per pool instance (so basically that we never access the DNA index from more than one thread). I didn't fully grasp why this is needed as the code was guarded with a critical section, but this was definitely another part of the problem.
* Gave the RigLogic task access to the trait's instance data. I think this should be something to integrate into the core system and be an automatic thing, all tasks should have access to the shared and instance trait data?
* Encapsulated the DNA index mapping init code from per skel mesh sharing functionality.

diff --git a/Engine/Plugins/Animation/RigLogic/Source/RigLogicModule/Private/DNAAsset.cpp b/Engine/Plugins/Animation/RigLogic/Source/RigLogicModule/Private/DNAAsset.cpp
--- a/Engine/Plugins/Animation/RigLogic/Source/RigLogicModule/Private/DNAAsset.cpp
+++ b/Engine/Plugins/Animation/RigLogic/Source/RigLogicModule/Private/DNAAsset.cpp
@@ -250,13 +250,7 @@ TSharedPtr<FDNAIndexMapping> UDNAAsset::GetDNAIndexMapping(const USkeleton* Skel
 	if (!DNAIndexMapping.IsValid() || (SkeletonGuid != DNAIndexMapping->SkeletonGuid))
 	{
 		DNAIndexMapping = MakeShared<FDNAIndexMapping>();
-		DNAIndexMapping->SkeletonGuid = SkeletonGuid;
-		DNAIndexMapping->MapControlCurves(BehaviorReader.Get(), Skeleton);
-		DNAIndexMapping->MapNeuralNetworkMaskCurves(BehaviorReader.Get(), Skeleton);
-		DNAIndexMapping->MapJoints(BehaviorReader.Get(), SkeletalMesh);
-		DNAIndexMapping->MapDriverJoints(BehaviorReader.Get(), SkeletalMesh);
-		DNAIndexMapping->MapMorphTargets(BehaviorReader.Get(), Skeleton, SkeletalMesh);
-		DNAIndexMapping->MapMaskMultipliers(BehaviorReader.Get(), Skeleton);
+		DNAIndexMapping->Init(BehaviorReader.Get(), Skeleton, SkeletalMesh);
 		DNAIndexMappingContainer.Add(SkeletalMesh, DNAIndexMapping);
 	}
 
diff --git a/Engine/Plugins/Animation/RigLogic/Source/RigLogicModule/Private/DNAIndexMapping.cpp b/Engine/Plugins/Animation/RigLogic/Source/RigLogicModule/Private/DNAIndexMapping.cpp
--- a/Engine/Plugins/Animation/RigLogic/Source/RigLogicModule/Private/DNAIndexMapping.cpp
+++ b/Engine/Plugins/Animation/RigLogic/Source/RigLogicModule/Private/DNAIndexMapping.cpp
@@ -220,3 +220,14 @@ void FDNAIndexMapping::MapMaskMultipliers(const IDNAReader* DNAReader, const USk
 		}
 	}
 }
+
+void FDNAIndexMapping::Init(const IDNAReader* DNAReader, const USkeleton* Skeleton, const USkeletalMesh* SkeletalMesh)
+{
+	SkeletonGuid = Skeleton->GetGuid();
+	MapControlCurves(DNAReader, Skeleton);
+	MapNeuralNetworkMaskCurves(DNAReader, Skeleton);
+	MapJoints(DNAReader, SkeletalMesh);
+	MapDriverJoints(DNAReader, SkeletalMesh);
+	MapMorphTargets(DNAReader, Skeleton, SkeletalMesh);
+	MapMaskMultipliers(DNAReader, Skeleton);
+}
\ No newline at end of file
diff --git a/Engine/Plugins/Animation/RigLogic/Source/RigLogicModule/Public/DNAIndexMapping.h b/Engine/Plugins/Animation/RigLogic/Source/RigLogicModule/Public/DNAIndexMapping.h
--- a/Engine/Plugins/Animation/RigLogic/Source/RigLogicModule/Public/DNAIndexMapping.h
+++ b/Engine/Plugins/Animation/RigLogic/Source/RigLogicModule/Public/DNAIndexMapping.h
@@ -11,7 +11,7 @@ class IDNAReader;
 class USkeleton;
 class USkeletalMesh;
 
-struct FDNAIndexMapping
+struct RIGLOGICMODULE_API FDNAIndexMapping
 {
 	template <typename T>
 	struct TArrayWrapper
@@ -41,6 +41,8 @@ struct FDNAIndexMapping
 	TArray<FCachedIndexedCurve> MorphTargetCurvesPerLOD;
 	TArray<FCachedIndexedCurve> MaskMultiplierCurvesPerLOD;
 
+	void Init(const IDNAReader* DNAReader, const USkeleton* Skeleton, const USkeletalMesh* SkeletalMesh);
+
 	void MapControlCurves(const IDNAReader* DNAReader, const USkeleton* Skeleton);
 	void MapNeuralNetworkMaskCurves(const IDNAReader* DNAReader, const USkeleton* Skeleton);
 	void MapDriverJoints(const IDNAReader* DNAReader, const USkeletalMesh* SkeletalMesh);
diff --git a/Engine/Plugins/Experimental/RigLogicAnimNext/Source/RigLogicAnimNext/Private/RigLogicInstanceData.cpp b/Engine/Plugins/Experimental/RigLogicAnimNext/Source/RigLogicAnimNext/Private/RigLogicInstanceData.cpp
--- a/Engine/Plugins/Experimental/RigLogicAnimNext/Source/RigLogicAnimNext/Private/RigLogicInstanceData.cpp
+++ b/Engine/Plugins/Experimental/RigLogicAnimNext/Source/RigLogicAnimNext/Private/RigLogicInstanceData.cpp
@@ -50,10 +50,10 @@ namespace UE::AnimNext
 		if (CachedRigRuntimeContext != SharedRigRuntimeContext)
 		{
 			CachedRigRuntimeContext = SharedRigRuntimeContext;
-			RigInstance = TUniquePtr<FRigInstance>(new FRigInstance(CachedRigRuntimeContext->RigLogic.Get()));
 		}
 
-		CachedDNAIndexMapping = DNAAsset->GetDNAIndexMapping(Skeleton, SkeletalMesh);
+		CachedDNAIndexMapping = MakeShared<FDNAIndexMapping>();
+		CachedDNAIndexMapping->Init(DNAAsset->GetBehaviorReader().Get(), Skeleton, SkeletalMesh);
 
 		InitBoneIndexMapping(ReferencePose);
 		InitSparseAndDenseDriverJointMapping(ReferencePose);
diff --git a/Engine/Plugins/Experimental/RigLogicAnimNext/Source/RigLogicAnimNext/Private/RigLogicTask.cpp b/Engine/Plugins/Experimental/RigLogicAnimNext/Source/RigLogicAnimNext/Private/RigLogicTask.cpp
--- a/Engine/Plugins/Experimental/RigLogicAnimNext/Source/RigLogicAnimNext/Private/RigLogicTask.cpp
+++ b/Engine/Plugins/Experimental/RigLogicAnimNext/Source/RigLogicAnimNext/Private/RigLogicTask.cpp
@@ -17,9 +17,11 @@
 
 DEFINE_STAT(STAT_AnimNext_Task_RigLogic);
 
-FAnimNextRigLogicTask FAnimNextRigLogicTask::Make()
+FAnimNextRigLogicTask FAnimNextRigLogicTask::Make(UE::AnimNext::FRigLogicTrait::FInstanceData* InstanceData)
 {
-	return {};
+	FAnimNextRigLogicTask Result;
+	Result.TraitInstanceData = InstanceData;
+	return Result;
 }
 
 void FAnimNextRigLogicTask::Execute(UE::AnimNext::FEvaluationVM& VM) const
@@ -36,11 +38,17 @@ void FAnimNextRigLogicTask::Execute(UE::AnimNext::FEvaluationVM& VM) const
 
 	// Lazy-init instance data.
 	FRigLogicAnimNextModule& Module = FModuleManager::GetModuleChecked<FRigLogicAnimNextModule>("RigLogicAnimNext");
-	TSharedPtr<UE::AnimNext::FRigLogicAnimNextInstanceData> NewInstanceData = Module.DataPool.RequestData(Keyframe->Pose.RefPose);
+	TSharedPtr<UE::AnimNext::FRigLogicAnimNextInstanceData> PoolInstanceData = Module.DataPool.RequestData(Keyframe->Pose.RefPose);
 
-	FRigInstance* RigInstance = NewInstanceData->RigInstance.Get();
-	const FSharedRigRuntimeContext* RigRuntimeContext = NewInstanceData->CachedRigRuntimeContext.Get();
-	const FDNAIndexMapping* DNAIndexMapping = NewInstanceData->CachedDNAIndexMapping.Get();
+	const FSharedRigRuntimeContext* RigRuntimeContext = PoolInstanceData->CachedRigRuntimeContext.Get();
+
+	if (!TraitInstanceData->RigInstance)
+	{
+		TraitInstanceData->RigInstance = TUniquePtr<FRigInstance>(new FRigInstance(RigRuntimeContext->RigLogic.Get()));
+	}
+	FRigInstance* RigInstance = TraitInstanceData->RigInstance.Get();
+	
+	const FDNAIndexMapping* DNAIndexMapping = PoolInstanceData->CachedDNAIndexMapping.Get();
 	if (RigInstance && RigRuntimeContext && DNAIndexMapping)
 	{
 		if (const FRigLogic* RigLogic = RigRuntimeContext->RigLogic.Get())
@@ -49,15 +57,15 @@ void FAnimNextRigLogicTask::Execute(UE::AnimNext::FEvaluationVM& VM) const
 
 			// 1. Feed input curves to RigLogic.
 			FBlendedCurve& Curves = Keyframe->Curves;
-			const TArrayView<const float> NeutralJointValues = NewInstanceData->CachedRigRuntimeContext->RigLogic->GetNeutralJointValues();
+			const TArrayView<const float> NeutralJointValues = PoolInstanceData->CachedRigRuntimeContext->RigLogic->GetNeutralJointValues();
 
 			UpdateControlCurves(RigInstance,
 				Curves,
 				DNAIndexMapping->ControlAttributeCurves,
 				DNAIndexMapping->NeuralNetworkMaskCurves,
 				NeutralJointValues,
-				NewInstanceData->SparseDriverJointsToControlAttributesMapPerLOD[LODLevel],
-				NewInstanceData->DenseDriverJointsToControlAttributesMapPerLOD[LODLevel],
+				PoolInstanceData->SparseDriverJointsToControlAttributesMapPerLOD[LODLevel],
+				PoolInstanceData->DenseDriverJointsToControlAttributesMapPerLOD[LODLevel],
 				Keyframe->Pose);
 
 			// 2. Evaluate RigLogic.
@@ -65,8 +73,8 @@ void FAnimNextRigLogicTask::Execute(UE::AnimNext::FEvaluationVM& VM) const
 			RigLogic->Calculate(RigInstance);
 
 			// 3. Update AnimNext bone transforms and curves with RigLogic output.
-			TArrayView<const uint16> VariableJointIndices = NewInstanceData->CachedRigRuntimeContext->VariableJointIndicesPerLOD[LODLevel].Values;
-			const TArrayView<const UE::AnimNext::FRigLogicBoneMapping> RigLogicIndexToMeshIndexMapping = NewInstanceData->RigLogicToSkeletonBoneIndexMappingPerLOD[LODLevel];
+			TArrayView<const uint16> VariableJointIndices = RigRuntimeContext->VariableJointIndicesPerLOD[LODLevel].Values;
+			const TArrayView<const UE::AnimNext::FRigLogicBoneMapping> RigLogicIndexToMeshIndexMapping = PoolInstanceData->RigLogicToSkeletonBoneIndexMappingPerLOD[LODLevel];
 			UpdateJoints(RigLogicIndexToMeshIndexMapping,
 				VariableJointIndices,
 				NeutralJointValues,
@@ -78,7 +86,7 @@ void FAnimNextRigLogicTask::Execute(UE::AnimNext::FEvaluationVM& VM) const
 		}
 	}
 
-	Module.DataPool.FreeData(Keyframe->Pose.RefPose->SkeletalMesh, NewInstanceData);
+	Module.DataPool.FreeData(Keyframe->Pose.RefPose->SkeletalMesh, PoolInstanceData);
 
 	// Push our result back
 	VM.PushValue(UE::AnimNext::KEYFRAME_STACK_NAME, MoveTemp(Keyframe));
diff --git a/Engine/Plugins/Experimental/RigLogicAnimNext/Source/RigLogicAnimNext/Private/RigLogicTrait.cpp b/Engine/Plugins/Experimental/RigLogicAnimNext/Source/RigLogicAnimNext/Private/RigLogicTrait.cpp
--- a/Engine/Plugins/Experimental/RigLogicAnimNext/Source/RigLogicAnimNext/Private/RigLogicTrait.cpp
+++ b/Engine/Plugins/Experimental/RigLogicAnimNext/Source/RigLogicAnimNext/Private/RigLogicTrait.cpp
@@ -43,7 +43,7 @@ namespace UE::AnimNext
 			// TODO: Only add the task in case we want to actually want to apply facial animation for the current LOD.
 			//if (LODLevel < LODThreshold)
 			{
-				Context.AppendTask(FAnimNextRigLogicTask::Make());
+				Context.AppendTask(FAnimNextRigLogicTask::Make(InstanceData));
 			}
 		}
 	}
diff --git a/Engine/Plugins/Experimental/RigLogicAnimNext/Source/RigLogicAnimNext/Public/RigLogicInstanceData.h b/Engine/Plugins/Experimental/RigLogicAnimNext/Source/RigLogicAnimNext/Public/RigLogicInstanceData.h
--- a/Engine/Plugins/Experimental/RigLogicAnimNext/Source/RigLogicAnimNext/Public/RigLogicInstanceData.h
+++ b/Engine/Plugins/Experimental/RigLogicAnimNext/Source/RigLogicAnimNext/Public/RigLogicInstanceData.h
@@ -37,9 +37,6 @@ namespace UE::AnimNext
 		/** Cached pointer to the DNA index mapping which is originally owned by UDNAAsset. */
 		TSharedPtr<FDNAIndexMapping> CachedDNAIndexMapping;
 
-		/** Actually cloned RigLogic instance owned by this class. */
-		TUniquePtr<FRigInstance> RigInstance;
-
 		/** Bone index mapping from a RigLogic joint index to the reference skeleton bone index, one per LOD level. */
 		TArray<TArray<FRigLogicBoneMapping>> RigLogicToSkeletonBoneIndexMappingPerLOD;
 
diff --git a/Engine/Plugins/Experimental/RigLogicAnimNext/Source/RigLogicAnimNext/Public/RigLogicTask.h b/Engine/Plugins/Experimental/RigLogicAnimNext/Source/RigLogicAnimNext/Public/RigLogicTask.h
--- a/Engine/Plugins/Experimental/RigLogicAnimNext/Source/RigLogicAnimNext/Public/RigLogicTask.h
+++ b/Engine/Plugins/Experimental/RigLogicAnimNext/Source/RigLogicAnimNext/Public/RigLogicTask.h
@@ -27,7 +27,7 @@ struct RIGLOGICANIMNEXT_API FAnimNextRigLogicTask : public FAnimNextEvaluationTa
 
 	DECLARE_ANIM_EVALUATION_TASK(FAnimNextRigLogicTask)
 
-	static FAnimNextRigLogicTask Make();
+	static FAnimNextRigLogicTask Make(UE::AnimNext::FRigLogicTrait::FInstanceData* InstanceData);
 
 	// Task entry point
 	virtual void Execute(UE::AnimNext::FEvaluationVM& VM) const override;
@@ -52,4 +52,6 @@ struct RIGLOGICANIMNEXT_API FAnimNextRigLogicTask : public FAnimNextEvaluationTa
 
 	void UpdateBlendShapeCurves(const FCachedIndexedCurve& MorphTargetCurves, const TArrayView<const float> BlendShapeValues, FBlendedCurve& OutputCurves) const;
 	void UpdateAnimMapCurves(const FCachedIndexedCurve& MaskMultiplierCurves, const TArrayView<const float> AnimMapOutputs, FBlendedCurve& OutputCurves) const;
+
+	UE::AnimNext::FRigLogicTrait::FInstanceData* TraitInstanceData = nullptr;
 };
\ No newline at end of file
diff --git a/Engine/Plugins/Experimental/RigLogicAnimNext/Source/RigLogicAnimNext/Public/RigLogicTrait.h b/Engine/Plugins/Experimental/RigLogicAnimNext/Source/RigLogicAnimNext/Public/RigLogicTrait.h
--- a/Engine/Plugins/Experimental/RigLogicAnimNext/Source/RigLogicAnimNext/Public/RigLogicTrait.h
+++ b/Engine/Plugins/Experimental/RigLogicAnimNext/Source/RigLogicAnimNext/Public/RigLogicTrait.h
@@ -4,6 +4,7 @@
 
 #include "CoreMinimal.h"
 
+#include "RigInstance.h"
 #include "TraitCore/Trait.h"
 #include "TraitInterfaces/IContinuousBlend.h"
 #include "TraitInterfaces/IEvaluate.h"
@@ -54,6 +55,9 @@ namespace UE::AnimNext
 		{
 			/** Input node from which we receive the input pose as well as the facial expression curves. */
 			FTraitPtr Input;
+
+			/** Actually cloned RigLogic instance owned by this class. */
+			TUniquePtr<FRigInstance> RigInstance;
 		};
 
 		// IEvaluate impl
