Static analysis fix for RigLogicTask

diff --git a/Engine/Plugins/Experimental/RigLogicAnimNext/Source/RigLogicAnimNext/Private/RigLogicTask.cpp b/Engine/Plugins/Experimental/RigLogicAnimNext/Source/RigLogicAnimNext/Private/RigLogicTask.cpp
--- a/Engine/Plugins/Experimental/RigLogicAnimNext/Source/RigLogicAnimNext/Private/RigLogicTask.cpp
+++ b/Engine/Plugins/Experimental/RigLogicAnimNext/Source/RigLogicAnimNext/Private/RigLogicTask.cpp
@@ -41,48 +41,50 @@ void FAnimNextRigLogicTask::Execute(UE::AnimNext::FEvaluationVM& VM) const
 	TSharedPtr<UE::AnimNext::FRigLogicAnimNextInstanceData> PoolInstanceData = Module.DataPool.RequestData(Keyframe->Pose.RefPose);
 
 	const FSharedRigRuntimeContext* RigRuntimeContext = PoolInstanceData->CachedRigRuntimeContext.Get();
-
-	if (!TraitInstanceData->RigInstance)
-	{
-		TraitInstanceData->RigInstance = TUniquePtr<FRigInstance>(new FRigInstance(RigRuntimeContext->RigLogic.Get()));
-	}
-	FRigInstance* RigInstance = TraitInstanceData->RigInstance.Get();
-	
-	const FDNAIndexMapping* DNAIndexMapping = PoolInstanceData->CachedDNAIndexMapping.Get();
-	if (RigInstance && RigRuntimeContext && DNAIndexMapping)
+	if (RigRuntimeContext)
 	{
-		if (const FRigLogic* RigLogic = RigRuntimeContext->RigLogic.Get())
+		if (!TraitInstanceData->RigInstance)
+		{
+			TraitInstanceData->RigInstance = TUniquePtr<FRigInstance>(new FRigInstance(RigRuntimeContext->RigLogic.Get()));
+		}
+		FRigInstance* RigInstance = TraitInstanceData->RigInstance.Get();
+
+		const FDNAIndexMapping* DNAIndexMapping = PoolInstanceData->CachedDNAIndexMapping.Get();
+		if (RigInstance && RigRuntimeContext && DNAIndexMapping)
 		{
-			const int32 LODLevel = Keyframe->Pose.LODLevel;
-
-			// 1. Feed input curves to RigLogic.
-			FBlendedCurve& Curves = Keyframe->Curves;
-			const TArrayView<const float> NeutralJointValues = PoolInstanceData->CachedRigRuntimeContext->RigLogic->GetNeutralJointValues();
-
-			UpdateControlCurves(RigInstance,
-				Curves,
-				DNAIndexMapping->ControlAttributeCurves,
-				DNAIndexMapping->NeuralNetworkMaskCurves,
-				NeutralJointValues,
-				PoolInstanceData->SparseDriverJointsToControlAttributesMapPerLOD[LODLevel],
-				PoolInstanceData->DenseDriverJointsToControlAttributesMapPerLOD[LODLevel],
-				Keyframe->Pose);
-
-			// 2. Evaluate RigLogic.
-			RigInstance->SetLOD(LODLevel);
-			RigLogic->Calculate(RigInstance);
-
-			// 3. Update AnimNext bone transforms and curves with RigLogic output.
-			TArrayView<const uint16> VariableJointIndices = RigRuntimeContext->VariableJointIndicesPerLOD[LODLevel].Values;
-			const TArrayView<const UE::AnimNext::FRigLogicBoneMapping> RigLogicIndexToMeshIndexMapping = PoolInstanceData->RigLogicToSkeletonBoneIndexMappingPerLOD[LODLevel];
-			UpdateJoints(RigLogicIndexToMeshIndexMapping,
-				VariableJointIndices,
-				NeutralJointValues,
-				RigInstance->GetJointOutputs(),
-				Keyframe->Pose);
-
-			UpdateBlendShapeCurves(DNAIndexMapping->MorphTargetCurvesPerLOD[LODLevel], RigInstance->GetBlendShapeOutputs(), Curves);
-			UpdateAnimMapCurves(DNAIndexMapping->MaskMultiplierCurvesPerLOD[LODLevel], RigInstance->GetAnimatedMapOutputs(), Curves);
+			if (const FRigLogic* RigLogic = RigRuntimeContext->RigLogic.Get())
+			{
+				const int32 LODLevel = Keyframe->Pose.LODLevel;
+
+				// 1. Feed input curves to RigLogic.
+				FBlendedCurve& Curves = Keyframe->Curves;
+				const TArrayView<const float> NeutralJointValues = PoolInstanceData->CachedRigRuntimeContext->RigLogic->GetNeutralJointValues();
+
+				UpdateControlCurves(RigInstance,
+					Curves,
+					DNAIndexMapping->ControlAttributeCurves,
+					DNAIndexMapping->NeuralNetworkMaskCurves,
+					NeutralJointValues,
+					PoolInstanceData->SparseDriverJointsToControlAttributesMapPerLOD[LODLevel],
+					PoolInstanceData->DenseDriverJointsToControlAttributesMapPerLOD[LODLevel],
+					Keyframe->Pose);
+
+				// 2. Evaluate RigLogic.
+				RigInstance->SetLOD(LODLevel);
+				RigLogic->Calculate(RigInstance);
+
+				// 3. Update AnimNext bone transforms and curves with RigLogic output.
+				TArrayView<const uint16> VariableJointIndices = RigRuntimeContext->VariableJointIndicesPerLOD[LODLevel].Values;
+				const TArrayView<const UE::AnimNext::FRigLogicBoneMapping> RigLogicIndexToMeshIndexMapping = PoolInstanceData->RigLogicToSkeletonBoneIndexMappingPerLOD[LODLevel];
+				UpdateJoints(RigLogicIndexToMeshIndexMapping,
+					VariableJointIndices,
+					NeutralJointValues,
+					RigInstance->GetJointOutputs(),
+					Keyframe->Pose);
+
+				UpdateBlendShapeCurves(DNAIndexMapping->MorphTargetCurvesPerLOD[LODLevel], RigInstance->GetBlendShapeOutputs(), Curves);
+				UpdateAnimMapCurves(DNAIndexMapping->MaskMultiplierCurvesPerLOD[LODLevel], RigInstance->GetAnimatedMapOutputs(), Curves);
+			}
 		}
 	}
 
