Sharing DNAIndexMapping across AnimNext RigLogic instances

diff --git a/Engine/Plugins/Experimental/RigLogicAnimNext/Source/RigLogicAnimNext/Private/RigLogicInstanceData.cpp b/Engine/Plugins/Experimental/RigLogicAnimNext/Source/RigLogicAnimNext/Private/RigLogicInstanceData.cpp
--- a/Engine/Plugins/Experimental/RigLogicAnimNext/Source/RigLogicAnimNext/Private/RigLogicInstanceData.cpp
+++ b/Engine/Plugins/Experimental/RigLogicAnimNext/Source/RigLogicAnimNext/Private/RigLogicInstanceData.cpp
@@ -13,6 +13,8 @@ namespace UE::AnimNext
 {
 	void FRigLogicAnimNextInstanceData::Init(const UE::AnimNext::FReferencePose* ReferencePose)
 	{
+		LLM_SCOPE_BYNAME(TEXT("Animation/RigLogic"));
+
 		if (!ReferencePose)
 		{
 			UE_LOG(LogRigLogicAnimNext, Error, TEXT("Reference pose invalid."));
@@ -52,8 +54,7 @@ namespace UE::AnimNext
 			CachedRigRuntimeContext = SharedRigRuntimeContext;
 		}
 
-		CachedDNAIndexMapping = MakeShared<FDNAIndexMapping>();
-		CachedDNAIndexMapping->Init(DNAAsset->GetBehaviorReader().Get(), Skeleton, SkeletalMesh);
+		CachedDNAIndexMapping = DNAAsset->GetDNAIndexMapping(Skeleton, SkeletalMesh);
 
 		NumLODs = CachedRigRuntimeContext->VariableJointIndicesPerLOD.Num();
 
@@ -63,6 +64,8 @@ namespace UE::AnimNext
 
 	void FRigLogicAnimNextInstanceData::InitBoneIndexMapping(const UE::AnimNext::FReferencePose* ReferencePose)
 	{
+		LLM_SCOPE_BYNAME(TEXT("Animation/RigLogic"));
+
 		RigLogicToSkeletonBoneIndexMappingPerLOD.Empty();
 		RigLogicToSkeletonBoneIndexMappingPerLOD.SetNum(NumLODs);
 
@@ -96,6 +99,8 @@ namespace UE::AnimNext
 
 	void FRigLogicAnimNextInstanceData::InitSparseAndDenseDriverJointMapping(const UE::AnimNext::FReferencePose* ReferencePose)
 	{
+		LLM_SCOPE_BYNAME(TEXT("Animation/RigLogic"));
+
 		// Populate driver joint to raw control attribute mapping (used to feed RigLogic with inputs from the joint hierarchy)
 		SparseDriverJointsToControlAttributesMapPerLOD.Empty();
 		DenseDriverJointsToControlAttributesMapPerLOD.Empty();
diff --git a/Engine/Plugins/Experimental/RigLogicAnimNext/Source/RigLogicAnimNext/Private/RigLogicInstanceDataPool.cpp b/Engine/Plugins/Experimental/RigLogicAnimNext/Source/RigLogicAnimNext/Private/RigLogicInstanceDataPool.cpp
--- a/Engine/Plugins/Experimental/RigLogicAnimNext/Source/RigLogicAnimNext/Private/RigLogicInstanceDataPool.cpp
+++ b/Engine/Plugins/Experimental/RigLogicAnimNext/Source/RigLogicAnimNext/Private/RigLogicInstanceDataPool.cpp
@@ -7,6 +7,8 @@ namespace UE::AnimNext
 {
 	TSharedPtr<FRigLogicAnimNextInstanceData> FRigLogicInstanceDataPool::RequestData(const UE::AnimNext::FReferencePose* InReferencePose)
 	{
+		LLM_SCOPE_BYNAME(TEXT("Animation/RigLogic"));
+
 		// Only lock while searching
 		{
 			FScopeLock Lock(&PoolAccessCriticalSection);
@@ -32,6 +34,7 @@ namespace UE::AnimNext
 
 	void FRigLogicInstanceDataPool::FreeData(TWeakObjectPtr<const USkeletalMesh> SkeletalMesh, TSharedPtr<FRigLogicAnimNextInstanceData> InData)
 	{
+		LLM_SCOPE_BYNAME(TEXT("Animation/RigLogic"));
 		FScopeLock Lock(&PoolAccessCriticalSection);
 
 		// Do we have an element for the given skel mesh?
@@ -52,6 +55,7 @@ namespace UE::AnimNext
 
 	void FRigLogicInstanceDataPool::GarbageCollect()
 	{
+		LLM_SCOPE_BYNAME(TEXT("Animation/RigLogic"));
 		FScopeLock Lock(&PoolAccessCriticalSection);
 
 		for (auto Iterator = Datas.CreateIterator(); Iterator; ++Iterator)
diff --git a/Engine/Plugins/Experimental/RigLogicAnimNext/Source/RigLogicAnimNext/Private/RigLogicTrait.cpp b/Engine/Plugins/Experimental/RigLogicAnimNext/Source/RigLogicAnimNext/Private/RigLogicTrait.cpp
--- a/Engine/Plugins/Experimental/RigLogicAnimNext/Source/RigLogicAnimNext/Private/RigLogicTrait.cpp
+++ b/Engine/Plugins/Experimental/RigLogicAnimNext/Source/RigLogicAnimNext/Private/RigLogicTrait.cpp
@@ -22,6 +22,8 @@ namespace UE::AnimNext
 
 	void FRigLogicTrait::OnBecomeRelevant(FUpdateTraversalContext& Context, const TTraitBinding<IUpdate>& Binding, const FTraitUpdateState& TraitState) const
 	{
+		LLM_SCOPE_BYNAME(TEXT("Animation/RigLogic"));
+
 		const FSharedData* SharedData = Binding.GetSharedData<FSharedData>();
 		FInstanceData* InstanceData = Binding.GetInstanceData<FInstanceData>();
 
