UAF: Namespace renamed

diff --git a/Engine/Plugins/Experimental/EvaluationNotifies/Source/Runtime/Private/AnimNext/AnimNextAlignment.h b/Engine/Plugins/Experimental/EvaluationNotifies/Source/Runtime/Private/AnimNext/AnimNextAlignment.h
--- a/Engine/Plugins/Experimental/EvaluationNotifies/Source/Runtime/Private/AnimNext/AnimNextAlignment.h
+++ b/Engine/Plugins/Experimental/EvaluationNotifies/Source/Runtime/Private/AnimNext/AnimNextAlignment.h
@@ -21,8 +21,8 @@ struct FEvaluationNotify_AlignmentInstance : public FEvaluationNotify_BaseInstan
 {
 	GENERATED_BODY()
 	virtual void Start() override;
-	virtual bool GetTargetTransform(UE::AnimNext::FEvaluationNotifiesTrait::FInstanceData& TraitInstanceData, FTransform& TargetTranform);
-   	virtual void Update(UE::AnimNext::FEvaluationNotifiesTrait::FInstanceData& InstanceData, UE::AnimNext::FEvaluationVM& VM) override;
+	virtual bool GetTargetTransform(UE::UAF::FEvaluationNotifiesTrait::FInstanceData& TraitInstanceData, FTransform& TargetTranform);
+   	virtual void Update(UE::UAF::FEvaluationNotifiesTrait::FInstanceData& InstanceData, UE::UAF::FEvaluationVM& VM) override;
 	
 	float GetWeight(float Time, const FAlignmentWarpCurve& WarpCurve) const; 
 
@@ -47,6 +47,6 @@ USTRUCT()
 struct FEvaluationNotify_AlignToGroundInstance : public FEvaluationNotify_AlignmentInstance
 {
 	GENERATED_BODY()
-	virtual void End(UE::AnimNext::FEvaluationNotifiesTrait::FInstanceData& TraitInstanceData) override;
-	virtual bool GetTargetTransform(UE::AnimNext::FEvaluationNotifiesTrait::FInstanceData& TraitInstanceData, FTransform& TargetTranform);
+	virtual void End(UE::UAF::FEvaluationNotifiesTrait::FInstanceData& TraitInstanceData) override;
+	virtual bool GetTargetTransform(UE::UAF::FEvaluationNotifiesTrait::FInstanceData& TraitInstanceData, FTransform& TargetTranform);
 };
\ No newline at end of file
diff --git a/Engine/Plugins/Experimental/EvaluationNotifies/Source/Runtime/Private/AnimNext/EvaluationNotifiesTrait.cpp b/Engine/Plugins/Experimental/EvaluationNotifies/Source/Runtime/Private/AnimNext/EvaluationNotifiesTrait.cpp
--- a/Engine/Plugins/Experimental/EvaluationNotifies/Source/Runtime/Private/AnimNext/EvaluationNotifiesTrait.cpp
+++ b/Engine/Plugins/Experimental/EvaluationNotifies/Source/Runtime/Private/AnimNext/EvaluationNotifiesTrait.cpp
@@ -12,7 +12,7 @@ static const FName RootBoneTransformName = "RootBoneTransform";
 ////////////////////////////////////////////////////////////////////////////////////////////////////////////////////////
 // FEvaluationNotifiesTrait
 
-namespace UE::AnimNext
+namespace UE::UAF
 {
 	TMap<UClass*, UScriptStruct*> FEvaluationNotifiesTrait::NotifyEvaluationHandlerMap;
 
@@ -152,13 +152,13 @@ namespace UE::AnimNext
 		FTrait::FInstanceData::Destruct(Context, Binding);
 	}
 
-} // namespace UE::AnimNext
+} // namespace UE::UAF
 
 
 ////////////////////////////////////////////////////////////////////////////////////////////////////////////////////////
 // FAnimNextEvaluationNotifiesTask
 
-FAnimNextEvaluationNotifiesTask FAnimNextEvaluationNotifiesTask::Make(UE::AnimNext::FEvaluationNotifiesTrait::FInstanceData* InstanceData, const UE::AnimNext::FEvaluationNotifiesTrait::FSharedData* SharedData)
+FAnimNextEvaluationNotifiesTask FAnimNextEvaluationNotifiesTask::Make(UE::UAF::FEvaluationNotifiesTrait::FInstanceData* InstanceData, const UE::UAF::FEvaluationNotifiesTrait::FSharedData* SharedData)
 {
 	FAnimNextEvaluationNotifiesTask Task;
 	Task.InstanceData = InstanceData;
@@ -166,9 +166,9 @@ FAnimNextEvaluationNotifiesTask FAnimNextEvaluationNotifiesTask::Make(UE::AnimNe
 	return Task;
 }
 
-void FAnimNextEvaluationNotifiesTask::Execute(UE::AnimNext::FEvaluationVM& VM) const
+void FAnimNextEvaluationNotifiesTask::Execute(UE::UAF::FEvaluationVM& VM) const
 {
-	using namespace UE::AnimNext;
+	using namespace UE::UAF;
 
 	// Execute EvaluationNotifies
 	if (InstanceData->DeltaTime > 0)
diff --git a/Engine/Plugins/Experimental/EvaluationNotifies/Source/Runtime/Private/EvaluationNotifiesRuntimeModule.cpp b/Engine/Plugins/Experimental/EvaluationNotifies/Source/Runtime/Private/EvaluationNotifiesRuntimeModule.cpp
--- a/Engine/Plugins/Experimental/EvaluationNotifies/Source/Runtime/Private/EvaluationNotifiesRuntimeModule.cpp
+++ b/Engine/Plugins/Experimental/EvaluationNotifies/Source/Runtime/Private/EvaluationNotifiesRuntimeModule.cpp
@@ -37,8 +37,8 @@ namespace UE
 			FAnimNode_EvaluationNotifies::RegisterEvaluationHandler(UNotifyState_Alignment::StaticClass(), FAlignmentNotifyInstance::StaticStruct());
 			FAnimNode_EvaluationNotifies::RegisterEvaluationHandler(UNotifyState_TwoBoneIK::StaticClass(), FTwoBoneIKNotifyInstance::StaticStruct());
 
-			AnimNext::FEvaluationNotifiesTrait::RegisterEvaluationHandler(UNotifyState_Alignment::StaticClass(), FEvaluationNotify_AlignmentInstance::StaticStruct());
-			AnimNext::FEvaluationNotifiesTrait::RegisterEvaluationHandler(UNotifyState_AlignToGround::StaticClass(), FEvaluationNotify_AlignToGroundInstance::StaticStruct());
+			UE::UAF::FEvaluationNotifiesTrait::RegisterEvaluationHandler(UNotifyState_Alignment::StaticClass(), FEvaluationNotify_AlignmentInstance::StaticStruct());
+			UE::UAF::FEvaluationNotifiesTrait::RegisterEvaluationHandler(UNotifyState_AlignToGround::StaticClass(), FEvaluationNotify_AlignToGroundInstance::StaticStruct());
 		}
 
 		void FModule::ShutdownModule()
@@ -46,7 +46,7 @@ namespace UE
 			FAnimNode_EvaluationNotifies::UnregisterEvaluationHandler(UNotifyState_Alignment::StaticClass());
 			FAnimNode_EvaluationNotifies::UnregisterEvaluationHandler(UNotifyState_TwoBoneIK::StaticClass());
 			
-			AnimNext::FEvaluationNotifiesTrait::UnregisterEvaluationHandler(UNotifyState_Alignment::StaticClass());
+			UE::UAF::FEvaluationNotifiesTrait::UnregisterEvaluationHandler(UNotifyState_Alignment::StaticClass());
 		}
 	}
 }
diff --git a/Engine/Plugins/Experimental/EvaluationNotifies/Source/Runtime/Public/AnimNext/EvaluationNotifiesTrait.h b/Engine/Plugins/Experimental/EvaluationNotifies/Source/Runtime/Public/AnimNext/EvaluationNotifiesTrait.h
--- a/Engine/Plugins/Experimental/EvaluationNotifies/Source/Runtime/Public/AnimNext/EvaluationNotifiesTrait.h
+++ b/Engine/Plugins/Experimental/EvaluationNotifies/Source/Runtime/Public/AnimNext/EvaluationNotifiesTrait.h
@@ -32,7 +32,7 @@ struct FEvaluationNotifiesTraitSharedData : public FAnimNextTraitSharedData
 	#undef TRAIT_LATENT_PROPERTIES_ENUMERATOR
 };
 
-namespace UE::AnimNext
+namespace UE::UAF
 {
 /**
  * this trait will run evaluation time code for notifies which have a registered handler
@@ -78,7 +78,7 @@ private:
 	static TMap<UClass*, UScriptStruct*> NotifyEvaluationHandlerMap;
 };
 
-} // namespace UE::AnimNext
+} // namespace UE::UAF
 
 /** Task to run EvaluationNotifies on VM */
 USTRUCT()
@@ -88,13 +88,13 @@ struct FAnimNextEvaluationNotifiesTask : public FAnimNextEvaluationTask
 
 	DECLARE_ANIM_EVALUATION_TASK(FAnimNextEvaluationNotifiesTask)
 
-	static FAnimNextEvaluationNotifiesTask Make(UE::AnimNext::FEvaluationNotifiesTrait::FInstanceData* InstanceData, const UE::AnimNext::FEvaluationNotifiesTrait::FSharedData* SharedData);
+	static FAnimNextEvaluationNotifiesTask Make(UE::UAF::FEvaluationNotifiesTrait::FInstanceData* InstanceData, const UE::UAF::FEvaluationNotifiesTrait::FSharedData* SharedData);
 
 	// Task entry point
-	virtual void Execute(UE::AnimNext::FEvaluationVM& VM) const override;
+	virtual void Execute(UE::UAF::FEvaluationVM& VM) const override;
 
-	UE::AnimNext::FEvaluationNotifiesTrait::FInstanceData* InstanceData = nullptr;
-	const UE::AnimNext::FEvaluationNotifiesTrait::FSharedData* SharedData = nullptr;
+	UE::UAF::FEvaluationNotifiesTrait::FInstanceData* InstanceData = nullptr;
+	const UE::UAF::FEvaluationNotifiesTrait::FSharedData* SharedData = nullptr;
 };
 
 USTRUCT()
@@ -102,8 +102,8 @@ struct FEvaluationNotify_BaseInstance
 {
 	GENERATED_BODY()
 	virtual void Start() {}
-	virtual void Update(UE::AnimNext::FEvaluationNotifiesTrait::FInstanceData& InstanceData, UE::AnimNext::FEvaluationVM& VM) {}
-	virtual void End(UE::AnimNext::FEvaluationNotifiesTrait::FInstanceData& InstanceData) {}
+	virtual void Update(UE::UAF::FEvaluationNotifiesTrait::FInstanceData& InstanceData, UE::UAF::FEvaluationVM& VM) {}
+	virtual void End(UE::UAF::FEvaluationNotifiesTrait::FInstanceData& InstanceData) {}
 	
 	virtual ~FEvaluationNotify_BaseInstance() {};
 
diff --git a/Engine/Plugins/Experimental/MoverAnimNext/Source/MoverAnimNext/Private/Module/RigVMTrait_ModuleEventDependency_MoverComponentTickFunctions.cpp b/Engine/Plugins/Experimental/MoverAnimNext/Source/MoverAnimNext/Private/Module/RigVMTrait_ModuleEventDependency_MoverComponentTickFunctions.cpp
--- a/Engine/Plugins/Experimental/MoverAnimNext/Source/MoverAnimNext/Private/Module/RigVMTrait_ModuleEventDependency_MoverComponentTickFunctions.cpp
+++ b/Engine/Plugins/Experimental/MoverAnimNext/Source/MoverAnimNext/Private/Module/RigVMTrait_ModuleEventDependency_MoverComponentTickFunctions.cpp
@@ -11,7 +11,7 @@ FString FRigVMTrait_ModuleEventDependency_MoverComponentTickFunctions::GetDispla
 }
 #endif // WITH_EDITOR
 
-void FRigVMTrait_ModuleEventDependency_MoverComponentTickFunctions::OnAddDependency(const UE::AnimNext::FModuleDependencyContext& InContext) const
+void FRigVMTrait_ModuleEventDependency_MoverComponentTickFunctions::OnAddDependency(const UE::UAF::FModuleDependencyContext& InContext) const
 {
 	UActorComponent* AnimNextComponent = Cast<UActorComponent>(InContext.Object);
 	if (AnimNextComponent == nullptr)
@@ -46,7 +46,7 @@ void FRigVMTrait_ModuleEventDependency_MoverComponentTickFunctions::OnAddDepende
 }
 
 
-void FRigVMTrait_ModuleEventDependency_MoverComponentTickFunctions::OnRemoveDependency(const UE::AnimNext::FModuleDependencyContext& InContext) const
+void FRigVMTrait_ModuleEventDependency_MoverComponentTickFunctions::OnRemoveDependency(const UE::UAF::FModuleDependencyContext& InContext) const
 {
 	UActorComponent* AnimNextComponent = Cast<UActorComponent>(InContext.Object);
 	if (AnimNextComponent == nullptr)
diff --git a/Engine/Plugins/Experimental/MovieSceneAnimMixer/Source/MovieSceneAnimMixer/Private/AnimNode_SequencerMixerTarget.cpp b/Engine/Plugins/Experimental/MovieSceneAnimMixer/Source/MovieSceneAnimMixer/Private/AnimNode_SequencerMixerTarget.cpp
--- a/Engine/Plugins/Experimental/MovieSceneAnimMixer/Source/MovieSceneAnimMixer/Private/AnimNode_SequencerMixerTarget.cpp
+++ b/Engine/Plugins/Experimental/MovieSceneAnimMixer/Source/MovieSceneAnimMixer/Private/AnimNode_SequencerMixerTarget.cpp
@@ -59,12 +59,12 @@ void FAnimNode_SequencerMixerTarget::Evaluate_AnyThread(FPoseContext& Output)
 	
 	int32 LODLevel = Output.AnimInstanceProxy->GetLODLevel();
 
-	using namespace UE::AnimNext;
+	using namespace UE::UAF;
 	
 	FDataHandle RefPoseHandle = FDataRegistry::Get()->GetOrGenerateReferencePose(SkeletalMeshComponent);
 	FAnimNextGraphReferencePose GraphReferencePose(RefPoseHandle);
 
-	const UE::AnimNext::FReferencePose& RefPose = RefPoseHandle.GetRef<UE::AnimNext::FReferencePose>();
+	const UE::UAF::FReferencePose& RefPose = RefPoseHandle.GetRef<UE::UAF::FReferencePose>();
 	FAnimNextGraphLODPose ResultPose;
 	ResultPose.LODPose = FLODPoseHeap(RefPose, LODLevel, true, Output.ExpectsAdditivePose());
 
diff --git a/Engine/Plugins/Experimental/MovieSceneAnimMixer/Source/MovieSceneAnimMixer/Private/Systems/MovieSceneAnimMixerSystem.cpp b/Engine/Plugins/Experimental/MovieSceneAnimMixer/Source/MovieSceneAnimMixer/Private/Systems/MovieSceneAnimMixerSystem.cpp
--- a/Engine/Plugins/Experimental/MovieSceneAnimMixer/Source/MovieSceneAnimMixer/Private/Systems/MovieSceneAnimMixerSystem.cpp
+++ b/Engine/Plugins/Experimental/MovieSceneAnimMixer/Source/MovieSceneAnimMixer/Private/Systems/MovieSceneAnimMixerSystem.cpp
@@ -110,7 +110,7 @@ namespace UE::MovieScene
 			const bool bRootComponentHasKeyedTransform = (Root && RootMotionSystem->IsTransformKeyed(Root));
 
 			// Create a new eval program
-			Mixer->EvaluationProgram = MakeShared<UE::AnimNext::FEvaluationProgram>();
+			Mixer->EvaluationProgram = MakeShared<UE::UAF::FEvaluationProgram>();
 
 			if (!MixerTask.IsValid())
 			{
@@ -341,9 +341,9 @@ FAnimNextBlendTwoKeyframesPreserveRootMotionTask FAnimNextBlendTwoKeyframesPrese
 	return Task;
 }
 
-void FAnimNextBlendTwoKeyframesPreserveRootMotionTask::Execute(UE::AnimNext::FEvaluationVM& VM) const
+void FAnimNextBlendTwoKeyframesPreserveRootMotionTask::Execute(UE::UAF::FEvaluationVM& VM) const
 {
-	using namespace UE::AnimNext;
+	using namespace UE::UAF;
 	using namespace UE::MovieScene;
 
 	const UE::Anim::FAttributeId RootMotionDeltaAttributeId(UE::Anim::IAnimRootMotionProvider::AttributeName, FCompactPoseBoneIndex(0));
@@ -469,9 +469,9 @@ FMovieSceneAccumulateAbsoluteBlendTask FMovieSceneAccumulateAbsoluteBlendTask::M
 	return Task;
 }
 
-void FMovieSceneAccumulateAbsoluteBlendTask::Execute(UE::AnimNext::FEvaluationVM& VM) const
+void FMovieSceneAccumulateAbsoluteBlendTask::Execute(UE::UAF::FEvaluationVM& VM) const
 {
-	using namespace UE::AnimNext;
+	using namespace UE::UAF;
 
 	// Pop our top two poses, we'll re-use the top keyframe for our result
 
diff --git a/Engine/Plugins/Experimental/MovieSceneAnimMixer/Source/MovieSceneAnimMixer/Private/Systems/MovieSceneAnimNextTargetSystem.cpp b/Engine/Plugins/Experimental/MovieSceneAnimMixer/Source/MovieSceneAnimMixer/Private/Systems/MovieSceneAnimNextTargetSystem.cpp
--- a/Engine/Plugins/Experimental/MovieSceneAnimMixer/Source/MovieSceneAnimMixer/Private/Systems/MovieSceneAnimNextTargetSystem.cpp
+++ b/Engine/Plugins/Experimental/MovieSceneAnimMixer/Source/MovieSceneAnimMixer/Private/Systems/MovieSceneAnimNextTargetSystem.cpp
@@ -55,7 +55,7 @@ namespace UE::MovieScene
 	struct FPreAnimatedAnimNextState
 	{
 		// Handle to the injection request instance (used to uninject when animation done)
-		UE::AnimNext::FInjectionRequestPtr InjectionHandle;
+		UE::UAF::FInjectionRequestPtr InjectionHandle;
 
 		TObjectPtr<UMovieSceneAnimNextTargetSystem> System;
 
@@ -78,7 +78,7 @@ namespace UE::MovieScene
 				{
 					InOutCachedValue.System->CurrentTargets.RemoveAll([&InOutCachedValue](const FMovieSceneAnimNextTargetData& TargetData) { return TargetData.InjectionRequestHandle == InOutCachedValue.InjectionHandle;});
 				}
-				UE::AnimNext::FInjectionUtils::Uninject(InOutCachedValue.InjectionHandle);
+				UE::UAF::FInjectionUtils::Uninject(InOutCachedValue.InjectionHandle);
 			}
 
 			// Reset the anim next module back to not being enabled if applicable
@@ -244,7 +244,7 @@ namespace UE::MovieScene
 
 			if (AnimNextComponent)
 			{
-				UE::AnimNext::FInjectionSite InjectionSite(InjectionTarget->InjectionSite);
+				UE::UAF::FInjectionSite InjectionSite(InjectionTarget->InjectionSite);
 				if (InjectionSite.DesiredSite.IsNone())
 				{
 					PRAGMA_DISABLE_DEPRECATION_WARNINGS
@@ -264,7 +264,7 @@ namespace UE::MovieScene
 					TargetData->Modifier = MakeShared<FMovieSceneAnimMixerEvaluationModifier>(MixerTask);
 
 					// Inject
-					TargetData->InjectionRequestHandle = UE::AnimNext::FInjectionUtils::InjectEvaluationModifier(AnimNextComponent, TargetData->Modifier.ToSharedRef(), TargetData->InjectionSite);
+					TargetData->InjectionRequestHandle = UE::UAF::FInjectionUtils::InjectEvaluationModifier(AnimNextComponent, TargetData->Modifier.ToSharedRef(), TargetData->InjectionSite);
 
 					// For now, we always restore state
 					PreAnimatedStorage->BeginTrackingEntity(EntityID, true, FRootInstanceHandle(), AnimNextComponent, InjectionSite.DesiredSite);
@@ -302,7 +302,7 @@ namespace UE::MovieScene
 
 } // namespace UE::MovieScene
 
-void FMovieSceneAnimMixerEvaluationModifier::PostEvaluate(UE::AnimNext::FEvaluateTraversalContext& Context) const
+void FMovieSceneAnimMixerEvaluationModifier::PostEvaluate(UE::UAF::FEvaluateTraversalContext& Context) const
 {
 	if (TaskToInject.IsValid())
 	{
@@ -371,7 +371,7 @@ void UMovieSceneAnimNextTargetSystem::OnUnlink()
 		// Remove injection
 		if (TargetData.InjectionRequestHandle.IsValid())
 		{
-			UE::AnimNext::FInjectionUtils::Uninject(TargetData.InjectionRequestHandle);
+			UE::UAF::FInjectionUtils::Uninject(TargetData.InjectionRequestHandle);
 		}
 	}
 
diff --git a/Engine/Plugins/Experimental/MovieSceneAnimMixer/Source/MovieSceneAnimMixer/Private/Systems/MovieSceneMixedSkeletalAnimationSystem.cpp b/Engine/Plugins/Experimental/MovieSceneAnimMixer/Source/MovieSceneAnimMixer/Private/Systems/MovieSceneMixedSkeletalAnimationSystem.cpp
--- a/Engine/Plugins/Experimental/MovieSceneAnimMixer/Source/MovieSceneAnimMixer/Private/Systems/MovieSceneMixedSkeletalAnimationSystem.cpp
+++ b/Engine/Plugins/Experimental/MovieSceneAnimMixer/Source/MovieSceneAnimMixer/Private/Systems/MovieSceneMixedSkeletalAnimationSystem.cpp
@@ -394,9 +394,9 @@ private:
 
 } // namespace UE::MovieScene
 
-void FMovieSceneSkeletalAnimationEvaluationTask::Execute(UE::AnimNext::FEvaluationVM& VM) const
+void FMovieSceneSkeletalAnimationEvaluationTask::Execute(UE::UAF::FEvaluationVM& VM) const
 {
-	using namespace UE::AnimNext;
+	using namespace UE::UAF;
 	using namespace UE::MovieScene;
 	FAnimMixerComponentTypes* AnimMixerComponents = FAnimMixerComponentTypes::Get();
 	if (const UAnimSequence* AnimSequencePtr = AnimationData.AnimSequence.Get())
diff --git a/Engine/Plugins/Experimental/MovieSceneAnimMixer/Source/MovieSceneAnimMixer/Private/Systems/MovieSceneRootMotionSystem.cpp b/Engine/Plugins/Experimental/MovieSceneAnimMixer/Source/MovieSceneAnimMixer/Private/Systems/MovieSceneRootMotionSystem.cpp
--- a/Engine/Plugins/Experimental/MovieSceneAnimMixer/Source/MovieSceneAnimMixer/Private/Systems/MovieSceneRootMotionSystem.cpp
+++ b/Engine/Plugins/Experimental/MovieSceneAnimMixer/Source/MovieSceneAnimMixer/Private/Systems/MovieSceneRootMotionSystem.cpp
@@ -265,9 +265,9 @@ FAnimNextConvertRootMotionToWorldSpaceTask FAnimNextConvertRootMotionToWorldSpac
 	return FAnimNextConvertRootMotionToWorldSpaceTask(InRootMotionData, InTransformOrigin, InRootTransform, InRootOffsetOrigin, InConversions);
 }
 
-void FAnimNextConvertRootMotionToWorldSpaceTask::Execute(UE::AnimNext::FEvaluationVM& VM) const
+void FAnimNextConvertRootMotionToWorldSpaceTask::Execute(UE::UAF::FEvaluationVM& VM) const
 {
-	using namespace UE::AnimNext; 
+	using namespace UE::UAF; 
 	using namespace UE::MovieScene;
 
 	if (!EnumHasAnyFlags(VM.GetFlags(), EEvaluationFlags::Attributes | EEvaluationFlags::Trajectory))
@@ -352,9 +352,9 @@ FAnimNextStoreRootTransformTask FAnimNextStoreRootTransformTask::Make(const TSha
 }
 
 
-void FAnimNextStoreRootTransformTask::Execute(UE::AnimNext::FEvaluationVM& VM) const
+void FAnimNextStoreRootTransformTask::Execute(UE::UAF::FEvaluationVM& VM) const
 {
-	using namespace UE::AnimNext;
+	using namespace UE::UAF;
 	using namespace UE::MovieScene;
 
 	FAnimMixerComponentTypes* AnimMixerComponents = FAnimMixerComponentTypes::Get();
diff --git a/Engine/Plugins/Experimental/MovieSceneAnimMixer/Source/MovieSceneAnimMixer/Public/MovieSceneAnimMixerModule.h b/Engine/Plugins/Experimental/MovieSceneAnimMixer/Source/MovieSceneAnimMixer/Public/MovieSceneAnimMixerModule.h
--- a/Engine/Plugins/Experimental/MovieSceneAnimMixer/Source/MovieSceneAnimMixer/Public/MovieSceneAnimMixerModule.h
+++ b/Engine/Plugins/Experimental/MovieSceneAnimMixer/Source/MovieSceneAnimMixer/Public/MovieSceneAnimMixerModule.h
@@ -9,12 +9,12 @@ MOVIESCENEANIMMIXER_API DECLARE_LOG_CATEGORY_EXTERN(LogMovieSceneAnimMixer, Log,
 
 struct FAnimationUpdateContext;
 
-namespace UE::AnimNext
+namespace UE::UAF
 {
 	struct FEvaluationVM;
 }
 
-DECLARE_MULTICAST_DELEGATE_TwoParams(FPreEvaluateMixerTasks, const FAnimationUpdateContext&, UE::AnimNext::FEvaluationVM&);
+DECLARE_MULTICAST_DELEGATE_TwoParams(FPreEvaluateMixerTasks, const FAnimationUpdateContext&, UE::UAF::FEvaluationVM&);
 
 namespace UE::MovieScene
 {
diff --git a/Engine/Plugins/Experimental/MovieSceneAnimMixer/Source/MovieSceneAnimMixer/Public/Systems/MovieSceneAnimMixerSystem.h b/Engine/Plugins/Experimental/MovieSceneAnimMixer/Source/MovieSceneAnimMixer/Public/Systems/MovieSceneAnimMixerSystem.h
--- a/Engine/Plugins/Experimental/MovieSceneAnimMixer/Source/MovieSceneAnimMixer/Public/Systems/MovieSceneAnimMixerSystem.h
+++ b/Engine/Plugins/Experimental/MovieSceneAnimMixer/Source/MovieSceneAnimMixer/Public/Systems/MovieSceneAnimMixerSystem.h
@@ -15,12 +15,9 @@
 
 struct FAnimNextEvaluationTask;
 
-namespace UE
+namespace UE::UAF
 {
-	namespace AnimNext
-	{
-		struct FEvaluationProgram;
-	}
+	struct FEvaluationProgram;
 }
 
 struct FMovieSceneAnimMixer;
@@ -84,7 +81,7 @@ struct FMovieSceneAnimMixer
 {
 	UE::MovieScene::FMovieSceneEntityID MixerEntityID;
 	TArray<TWeakPtr<FMovieSceneAnimMixerEntry>> WeakEntries;
-	TSharedPtr<UE::AnimNext::FEvaluationProgram> EvaluationProgram;
+	TSharedPtr<UE::UAF::FEvaluationProgram> EvaluationProgram;
 	TWeakPtr<FMovieSceneMixerRootMotionComponentData> WeakRootMotion;
 
 	bool bNeedsResort = false;
@@ -111,7 +108,7 @@ struct FMovieSceneAccumulateAbsoluteBlendTask : public FAnimNextEvaluationTask
 	static FMovieSceneAccumulateAbsoluteBlendTask Make(float ScaleFactor);
 
 	// Task entry point
-	virtual void Execute(UE::AnimNext::FEvaluationVM& VM) const override;
+	virtual void Execute(UE::UAF::FEvaluationVM& VM) const override;
 
 	// The scale factor to apply to the keyframe
 	UPROPERTY()
@@ -128,7 +125,7 @@ struct FAnimNextBlendTwoKeyframesPreserveRootMotionTask : public FAnimNextBlendT
 	static FAnimNextBlendTwoKeyframesPreserveRootMotionTask Make(float InterpolationAlpha);
 
 	// Task entry point
-	virtual void Execute(UE::AnimNext::FEvaluationVM& VM) const override;
+	virtual void Execute(UE::UAF::FEvaluationVM& VM) const override;
 };
 
 
diff --git a/Engine/Plugins/Experimental/MovieSceneAnimMixer/Source/MovieSceneAnimMixer/Public/Systems/MovieSceneAnimNextTargetSystem.h b/Engine/Plugins/Experimental/MovieSceneAnimMixer/Source/MovieSceneAnimMixer/Public/Systems/MovieSceneAnimNextTargetSystem.h
--- a/Engine/Plugins/Experimental/MovieSceneAnimMixer/Source/MovieSceneAnimMixer/Public/Systems/MovieSceneAnimNextTargetSystem.h
+++ b/Engine/Plugins/Experimental/MovieSceneAnimMixer/Source/MovieSceneAnimMixer/Public/Systems/MovieSceneAnimNextTargetSystem.h
@@ -13,7 +13,7 @@
 
 struct FAnimNextEvaluationTask;
 
-struct FMovieSceneAnimMixerEvaluationModifier : public UE::AnimNext::IEvaluationModifier
+struct FMovieSceneAnimMixerEvaluationModifier : public UE::UAF::IEvaluationModifier
 {
 	FMovieSceneAnimMixerEvaluationModifier(TSharedPtr<FAnimNextEvaluationTask> InTaskToInject)
 		: TaskToInject(InTaskToInject)
@@ -21,9 +21,9 @@ struct FMovieSceneAnimMixerEvaluationModifier : public UE::AnimNext::IEvaluation
 	}
 
 	// IEvaluationModifier interface
-	virtual void PreEvaluate(UE::AnimNext::FEvaluateTraversalContext& Context) const override {}
+	virtual void PreEvaluate(UE::UAF::FEvaluateTraversalContext& Context) const override {}
 
-	virtual void PostEvaluate(UE::AnimNext::FEvaluateTraversalContext& Context) const override;
+	virtual void PostEvaluate(UE::UAF::FEvaluateTraversalContext& Context) const override;
 
 	TSharedPtr<FAnimNextEvaluationTask> TaskToInject = nullptr;
 
@@ -58,8 +58,8 @@ struct MOVIESCENEANIMMIXER_API FMovieSceneAnimNextInjectionTarget : public FMovi
 struct FMovieSceneAnimNextTargetData
 {
 	TWeakObjectPtr<UAnimNextComponent> AnimNextComponent = nullptr;
-	UE::AnimNext::FInjectionSite InjectionSite;
-	UE::AnimNext::FInjectionRequestPtr InjectionRequestHandle = nullptr;
+	UE::UAF::FInjectionSite InjectionSite;
+	UE::UAF::FInjectionRequestPtr InjectionRequestHandle = nullptr;
 	TSharedPtr<FMovieSceneAnimMixerEvaluationModifier> Modifier = nullptr;
 };
 
diff --git a/Engine/Plugins/Experimental/MovieSceneAnimMixer/Source/MovieSceneAnimMixer/Public/Systems/MovieSceneMixedSkeletalAnimationSystem.h b/Engine/Plugins/Experimental/MovieSceneAnimMixer/Source/MovieSceneAnimMixer/Public/Systems/MovieSceneMixedSkeletalAnimationSystem.h
--- a/Engine/Plugins/Experimental/MovieSceneAnimMixer/Source/MovieSceneAnimMixer/Public/Systems/MovieSceneMixedSkeletalAnimationSystem.h
+++ b/Engine/Plugins/Experimental/MovieSceneAnimMixer/Source/MovieSceneAnimMixer/Public/Systems/MovieSceneMixedSkeletalAnimationSystem.h
@@ -48,7 +48,7 @@ struct FMovieSceneSkeletalAnimationEvaluationTask : public FAnimNextEvaluationTa
 	DECLARE_ANIM_EVALUATION_TASK(FMovieSceneSkeletalAnimationEvaluationTask)
 
 	// Task entry point
-	virtual void Execute(UE::AnimNext::FEvaluationVM& VM) const override;
+	virtual void Execute(UE::UAF::FEvaluationVM& VM) const override;
 
 	UPROPERTY()
 	FMixedAnimSkeletalAnimationData AnimationData;
diff --git a/Engine/Plugins/Experimental/MovieSceneAnimMixer/Source/MovieSceneAnimMixer/Public/Systems/MovieSceneRootMotionSystem.h b/Engine/Plugins/Experimental/MovieSceneAnimMixer/Source/MovieSceneAnimMixer/Public/Systems/MovieSceneRootMotionSystem.h
--- a/Engine/Plugins/Experimental/MovieSceneAnimMixer/Source/MovieSceneAnimMixer/Public/Systems/MovieSceneRootMotionSystem.h
+++ b/Engine/Plugins/Experimental/MovieSceneAnimMixer/Source/MovieSceneAnimMixer/Public/Systems/MovieSceneRootMotionSystem.h
@@ -123,7 +123,7 @@ struct FAnimNextConvertRootMotionToWorldSpaceTask : public FAnimNextEvaluationTa
 
 	static FAnimNextConvertRootMotionToWorldSpaceTask Make(const TSharedPtr<FMovieSceneMixerRootMotionComponentData>& InRootMotionData, const FTransform& InTransformOrigin, const FTransform& InRootTransform, const FVector& InRootOffsetOrigin, ESpaceConversions InConversions);
 
-	virtual void Execute(UE::AnimNext::FEvaluationVM& VM) const override;
+	virtual void Execute(UE::UAF::FEvaluationVM& VM) const override;
 
 	/** Weak pointer to the root motion data for all mixed animations. May be null if only a transform origin transformation is required. */
 	TWeakPtr<FMovieSceneMixerRootMotionComponentData> WeakRootMotionData;
@@ -153,7 +153,7 @@ struct FAnimNextStoreRootTransformTask : public FAnimNextEvaluationTask
 
 	static FAnimNextStoreRootTransformTask Make(const TSharedPtr<FMovieSceneMixerRootMotionComponentData>& InRootMotionData, bool bInComponentHasKeyedTransform, bool bInRootComponentHasKeyedTransform);
 
-	virtual void Execute(UE::AnimNext::FEvaluationVM& VM) const override;
+	virtual void Execute(UE::UAF::FEvaluationVM& VM) const override;
 
 	TWeakPtr<FMovieSceneMixerRootMotionComponentData> WeakRootMotionData;
 	bool bComponentHasKeyedTransform = false;
diff --git a/Engine/Plugins/Experimental/MovieScenePoseSearchTracks/Source/MovieScenePoseSearchTracks/Private/Sections/MovieSceneStitchAnimSection.cpp b/Engine/Plugins/Experimental/MovieScenePoseSearchTracks/Source/MovieScenePoseSearchTracks/Private/Sections/MovieSceneStitchAnimSection.cpp
--- a/Engine/Plugins/Experimental/MovieScenePoseSearchTracks/Source/MovieScenePoseSearchTracks/Private/Sections/MovieSceneStitchAnimSection.cpp
+++ b/Engine/Plugins/Experimental/MovieScenePoseSearchTracks/Source/MovieScenePoseSearchTracks/Private/Sections/MovieSceneStitchAnimSection.cpp
@@ -49,9 +49,9 @@ namespace UE::MovieScene
 	};
 } // namespace UE::MovieScene
 
-void FMovieSceneStitchAnimEvaluationTask::Execute(UE::AnimNext::FEvaluationVM& VM) const
+void FMovieSceneStitchAnimEvaluationTask::Execute(UE::UAF::FEvaluationVM& VM) const
 {
-	using namespace UE::AnimNext;
+	using namespace UE::UAF;
 	using namespace UE::MovieScene;
 	using namespace UE::PoseSearch;
 	using namespace UE::Anim;
diff --git a/Engine/Plugins/Experimental/MovieScenePoseSearchTracks/Source/MovieScenePoseSearchTracks/Public/Sections/MovieSceneStitchAnimSection.h b/Engine/Plugins/Experimental/MovieScenePoseSearchTracks/Source/MovieScenePoseSearchTracks/Public/Sections/MovieSceneStitchAnimSection.h
--- a/Engine/Plugins/Experimental/MovieScenePoseSearchTracks/Source/MovieScenePoseSearchTracks/Public/Sections/MovieSceneStitchAnimSection.h
+++ b/Engine/Plugins/Experimental/MovieScenePoseSearchTracks/Source/MovieScenePoseSearchTracks/Public/Sections/MovieSceneStitchAnimSection.h
@@ -55,7 +55,7 @@ struct FMovieSceneStitchAnimEvaluationTask : public FAnimNextEvaluationTask
 	DECLARE_ANIM_EVALUATION_TASK(FMovieSceneStitchAnimEvaluationTask)
 
 	// Task entry point
-	virtual void Execute(UE::AnimNext::FEvaluationVM& VM) const override;
+	virtual void Execute(UE::UAF::FEvaluationVM& VM) const override;
 
 	UPROPERTY(Transient)
 	FMovieSceneStitchAnimComponentData StitchData;
diff --git a/Engine/Plugins/Experimental/RigLogicAnimNext/Source/RigLogicAnimNext/Private/RigLogicInstanceData.cpp b/Engine/Plugins/Experimental/RigLogicAnimNext/Source/RigLogicAnimNext/Private/RigLogicInstanceData.cpp
--- a/Engine/Plugins/Experimental/RigLogicAnimNext/Source/RigLogicAnimNext/Private/RigLogicInstanceData.cpp
+++ b/Engine/Plugins/Experimental/RigLogicAnimNext/Source/RigLogicAnimNext/Private/RigLogicInstanceData.cpp
@@ -9,9 +9,9 @@
 #include "LODPose.h"
 #include "SharedRigRuntimeContext.h"
 
-namespace UE::AnimNext
+namespace UE::UAF
 {
-	void FRigLogicAnimNextInstanceData::Init(const UE::AnimNext::FReferencePose* ReferencePose)
+	void FRigLogicAnimNextInstanceData::Init(const UE::UAF::FReferencePose* ReferencePose)
 	{
 		LLM_SCOPE_BYNAME(TEXT("Animation/RigLogic"));
 
@@ -62,7 +62,7 @@ namespace UE::AnimNext
 		InitSparseAndDenseDriverJointMapping(ReferencePose);
 	}
 
-	void FRigLogicAnimNextInstanceData::InitBoneIndexMapping(const UE::AnimNext::FReferencePose* ReferencePose)
+	void FRigLogicAnimNextInstanceData::InitBoneIndexMapping(const UE::UAF::FReferencePose* ReferencePose)
 	{
 		LLM_SCOPE_BYNAME(TEXT("Animation/RigLogic"));
 
@@ -97,7 +97,7 @@ namespace UE::AnimNext
 		}
 	}
 
-	void FRigLogicAnimNextInstanceData::InitSparseAndDenseDriverJointMapping(const UE::AnimNext::FReferencePose* ReferencePose)
+	void FRigLogicAnimNextInstanceData::InitSparseAndDenseDriverJointMapping(const UE::UAF::FReferencePose* ReferencePose)
 	{
 		LLM_SCOPE_BYNAME(TEXT("Animation/RigLogic"));
 
@@ -135,4 +135,4 @@ namespace UE::AnimNext
 			}
 		}
 	}
-} // namespace UE::AnimNext
+} // namespace UE::UAF
diff --git a/Engine/Plugins/Experimental/RigLogicAnimNext/Source/RigLogicAnimNext/Private/RigLogicInstanceDataPool.cpp b/Engine/Plugins/Experimental/RigLogicAnimNext/Source/RigLogicAnimNext/Private/RigLogicInstanceDataPool.cpp
--- a/Engine/Plugins/Experimental/RigLogicAnimNext/Source/RigLogicAnimNext/Private/RigLogicInstanceDataPool.cpp
+++ b/Engine/Plugins/Experimental/RigLogicAnimNext/Source/RigLogicAnimNext/Private/RigLogicInstanceDataPool.cpp
@@ -3,9 +3,9 @@
 #include "RigLogicInstanceDataPool.h"
 #include "RigLogicAnimNext.h"
 
-namespace UE::AnimNext
+namespace UE::UAF
 {
-	TSharedPtr<FRigLogicAnimNextInstanceData> FRigLogicInstanceDataPool::RequestData(const UE::AnimNext::FReferencePose* InReferencePose)
+	TSharedPtr<FRigLogicAnimNextInstanceData> FRigLogicInstanceDataPool::RequestData(const UE::UAF::FReferencePose* InReferencePose)
 	{
 		LLM_SCOPE_BYNAME(TEXT("Animation/RigLogic"));
 
@@ -94,4 +94,4 @@ namespace UE::AnimNext
 			Counter++;
 		}
 	}
-} // namespace UE::AnimNext
\ No newline at end of file
+} // namespace UE::UAF
\ No newline at end of file
diff --git a/Engine/Plugins/Experimental/RigLogicAnimNext/Source/RigLogicAnimNext/Private/RigLogicTask.cpp b/Engine/Plugins/Experimental/RigLogicAnimNext/Source/RigLogicAnimNext/Private/RigLogicTask.cpp
--- a/Engine/Plugins/Experimental/RigLogicAnimNext/Source/RigLogicAnimNext/Private/RigLogicTask.cpp
+++ b/Engine/Plugins/Experimental/RigLogicAnimNext/Source/RigLogicAnimNext/Private/RigLogicTask.cpp
@@ -17,20 +17,20 @@
 
 DEFINE_STAT(STAT_AnimNext_Task_RigLogic);
 
-FAnimNextRigLogicTask FAnimNextRigLogicTask::Make(UE::AnimNext::FRigLogicTrait::FInstanceData* InstanceData)
+FAnimNextRigLogicTask FAnimNextRigLogicTask::Make(UE::UAF::FRigLogicTrait::FInstanceData* InstanceData)
 {
 	FAnimNextRigLogicTask Result;
 	Result.TraitInstanceData = InstanceData;
 	return Result;
 }
 
-void FAnimNextRigLogicTask::Execute(UE::AnimNext::FEvaluationVM& VM) const
+void FAnimNextRigLogicTask::Execute(UE::UAF::FEvaluationVM& VM) const
 {
 	SCOPE_CYCLE_COUNTER(STAT_AnimNext_Task_RigLogic);
 
 	// Pop pose, we'll re-use it for our result
-	TUniquePtr<UE::AnimNext::FKeyframeState> Keyframe;
-	if (!VM.PopValue(UE::AnimNext::KEYFRAME_STACK_NAME, Keyframe))
+	TUniquePtr<UE::UAF::FKeyframeState> Keyframe;
+	if (!VM.PopValue(UE::UAF::KEYFRAME_STACK_NAME, Keyframe))
 	{
 		// We have no inputs, nothing to do.
 		return;
@@ -38,7 +38,7 @@ void FAnimNextRigLogicTask::Execute(UE::AnimNext::FEvaluationVM& VM) const
 
 	// Lazy-init instance data.
 	FRigLogicAnimNextModule& Module = FModuleManager::GetModuleChecked<FRigLogicAnimNextModule>("RigLogicAnimNext");
-	TSharedPtr<UE::AnimNext::FRigLogicAnimNextInstanceData> PoolInstanceData = Module.DataPool.RequestData(Keyframe->Pose.RefPose);
+	TSharedPtr<UE::UAF::FRigLogicAnimNextInstanceData> PoolInstanceData = Module.DataPool.RequestData(Keyframe->Pose.RefPose);
 
 	const FSharedRigRuntimeContext* RigRuntimeContext = PoolInstanceData->CachedRigRuntimeContext.Get();
 	if (RigRuntimeContext)
@@ -75,7 +75,7 @@ void FAnimNextRigLogicTask::Execute(UE::AnimNext::FEvaluationVM& VM) const
 
 				// 3. Update AnimNext bone transforms and curves with RigLogic output.
 				TArrayView<const uint16> VariableJointIndices = RigRuntimeContext->VariableJointIndicesPerLOD[LODLevel].Values;
-				const TArrayView<const UE::AnimNext::FRigLogicBoneMapping> RigLogicIndexToMeshIndexMapping = PoolInstanceData->RigLogicToSkeletonBoneIndexMappingPerLOD[LODLevel];
+				const TArrayView<const UE::UAF::FRigLogicBoneMapping> RigLogicIndexToMeshIndexMapping = PoolInstanceData->RigLogicToSkeletonBoneIndexMappingPerLOD[LODLevel];
 				UpdateJoints(RigLogicIndexToMeshIndexMapping,
 					VariableJointIndices,
 					NeutralJointValues,
@@ -91,7 +91,7 @@ void FAnimNextRigLogicTask::Execute(UE::AnimNext::FEvaluationVM& VM) const
 	Module.DataPool.FreeData(Keyframe->Pose.RefPose->SkeletalMesh, PoolInstanceData);
 
 	// Push our result back
-	VM.PushValue(UE::AnimNext::KEYFRAME_STACK_NAME, MoveTemp(Keyframe));
+	VM.PushValue(UE::UAF::KEYFRAME_STACK_NAME, MoveTemp(Keyframe));
 }
 
 void FAnimNextRigLogicTask::UpdateControlCurves(FRigInstance* InRigInstance,
@@ -99,9 +99,9 @@ void FAnimNextRigLogicTask::UpdateControlCurves(FRigInstance* InRigInstance,
 	const FCachedIndexedCurve& ControlAttributeCurves,
 	const FCachedIndexedCurve& NeuralNetworkMaskCurves,
 	const TArrayView<const float> NeutralJointValues,
-	const TArrayView<const UE::AnimNext::FPoseBoneControlAttributeMapping> SparseDriverJointsToControlAttributesMap,
-	const TArrayView<const UE::AnimNext::FPoseBoneControlAttributeMapping> DenseDriverJointsToControlAttributesMap,
-	const UE::AnimNext::FLODPoseStack& InputPose) const
+	const TArrayView<const UE::UAF::FPoseBoneControlAttributeMapping> SparseDriverJointsToControlAttributesMap,
+	const TArrayView<const UE::UAF::FPoseBoneControlAttributeMapping> DenseDriverJointsToControlAttributesMap,
+	const UE::UAF::FLODPoseStack& InputPose) const
 {
 	// Combine control attribute curve with input curve to get indexed curve to apply to rig
 	// Curve elements that dont have a control mapping will have INDEX_NONE as their index
@@ -131,10 +131,10 @@ void FAnimNextRigLogicTask::UpdateControlCurves(FRigInstance* InRigInstance,
 	// The sparse mapping is NOT guaranteed to supply all quaternion attributes, so checks for each attribute mapping are present
 	for (int32 MappingIndex = 0; MappingIndex < SparseDriverJointsToControlAttributesMap.Num(); ++MappingIndex)
 	{
-		const UE::AnimNext::FPoseBoneControlAttributeMapping& Mapping = SparseDriverJointsToControlAttributesMap[MappingIndex];
+		const UE::UAF::FPoseBoneControlAttributeMapping& Mapping = SparseDriverJointsToControlAttributesMap[MappingIndex];
 		const int32 PoseBoneIndex = Mapping.PoseBoneIndex;
 
-		const UE::AnimNext::FTransformSoAAdapterConst TransformAdapter = InputPose.LocalTransforms[PoseBoneIndex];
+		const UE::UAF::FTransformSoAAdapterConst TransformAdapter = InputPose.LocalTransforms[PoseBoneIndex];
 
 		// Translation and Scale is currently not used here, so to avoid the overhead of checking them, they are simply ignored.
 		// Should the need arise to use them as well, this code will need adjustment.
@@ -166,10 +166,10 @@ void FAnimNextRigLogicTask::UpdateControlCurves(FRigInstance* InRigInstance,
 	// The dense mapping is guaranteed to supply all quaternion attributes, so NO checks for each attribute mapping are present
 	for (int32 MappingIndex = 0; MappingIndex < DenseDriverJointsToControlAttributesMap.Num(); ++MappingIndex)
 	{
-		const UE::AnimNext::FPoseBoneControlAttributeMapping& Mapping = DenseDriverJointsToControlAttributesMap[MappingIndex];
+		const UE::UAF::FPoseBoneControlAttributeMapping& Mapping = DenseDriverJointsToControlAttributesMap[MappingIndex];
 		const int32 PoseBoneIndex = Mapping.PoseBoneIndex;
 
-		const UE::AnimNext::FTransformSoAAdapterConst TransformAdapter = InputPose.LocalTransforms[PoseBoneIndex];
+		const UE::UAF::FTransformSoAAdapterConst TransformAdapter = InputPose.LocalTransforms[PoseBoneIndex];
 
 		// Translation and Scale is currently not used here, so to avoid the overhead of checking them, they are simply ignored.
 		// Should the need arise to use them as well, this code will need adjustment.
@@ -199,24 +199,24 @@ void FAnimNextRigLogicTask::UpdateControlCurves(FRigInstance* InRigInstance,
 	}
 }
 
-void FAnimNextRigLogicTask::UpdateJoints(const TArrayView<const UE::AnimNext::FRigLogicBoneMapping> RigLogicIndexToMeshIndexMapping,
+void FAnimNextRigLogicTask::UpdateJoints(const TArrayView<const UE::UAF::FRigLogicBoneMapping> RigLogicIndexToMeshIndexMapping,
 	const TArrayView<const uint16> VariableJointIndices,
 	const TArrayView<const float> NeutralJointValues,
 	const TArrayView<const float> DeltaJointValues,
-	UE::AnimNext::FLODPoseStack& OutputPose) const
+	UE::UAF::FLODPoseStack& OutputPose) const
 {
 	const TArrayView<const FBoneIndexType> SkeletonIndexToPoseIndex = OutputPose.GetSkeletonBoneIndexToLODBoneIndexMap(); // skeleton -> pose
 
 	const float* N = NeutralJointValues.GetData();
 	const float* D = DeltaJointValues.GetData();
 
-	for (const UE::AnimNext::FRigLogicBoneMapping& Mapping : RigLogicIndexToMeshIndexMapping)
+	for (const UE::UAF::FRigLogicBoneMapping& Mapping : RigLogicIndexToMeshIndexMapping)
 	{
 		const uint16 RigLogicJointIndex = Mapping.RigLogicJointIndex;
 		const int32 PoseBoneIndex = Mapping.PoseBoneIndex;
 		const uint16 AttrIndex = RigLogicJointIndex * ATTR_COUNT_PER_JOINT;
 
-		UE::AnimNext::FTransformSoAAdapter TransformAdapter = OutputPose.LocalTransforms[PoseBoneIndex];
+		UE::UAF::FTransformSoAAdapter TransformAdapter = OutputPose.LocalTransforms[PoseBoneIndex];
 
 		TransformAdapter.SetTranslation(FVector((N[AttrIndex + 0] + D[AttrIndex + 0]), (N[AttrIndex + 1] + D[AttrIndex + 1]), (N[AttrIndex + 2] + D[AttrIndex + 2])));
 		TransformAdapter.SetRotation(FQuat(N[AttrIndex + 3], N[AttrIndex + 4], N[AttrIndex + 5], N[AttrIndex + 6]) * FQuat(D[AttrIndex + 3], D[AttrIndex + 4], D[AttrIndex + 5], D[AttrIndex + 6]));
diff --git a/Engine/Plugins/Experimental/RigLogicAnimNext/Source/RigLogicAnimNext/Public/RigLogicAnimNext.h b/Engine/Plugins/Experimental/RigLogicAnimNext/Source/RigLogicAnimNext/Public/RigLogicAnimNext.h
--- a/Engine/Plugins/Experimental/RigLogicAnimNext/Source/RigLogicAnimNext/Public/RigLogicAnimNext.h
+++ b/Engine/Plugins/Experimental/RigLogicAnimNext/Source/RigLogicAnimNext/Public/RigLogicAnimNext.h
@@ -14,5 +14,5 @@ public:
 	virtual void StartupModule() override;
 	virtual void ShutdownModule() override;
 
-	UE::AnimNext::FRigLogicInstanceDataPool DataPool;
+	UE::UAF::FRigLogicInstanceDataPool DataPool;
 };
\ No newline at end of file
diff --git a/Engine/Plugins/Experimental/RigLogicAnimNext/Source/RigLogicAnimNext/Public/RigLogicInstanceData.h b/Engine/Plugins/Experimental/RigLogicAnimNext/Source/RigLogicAnimNext/Public/RigLogicInstanceData.h
--- a/Engine/Plugins/Experimental/RigLogicAnimNext/Source/RigLogicAnimNext/Public/RigLogicInstanceData.h
+++ b/Engine/Plugins/Experimental/RigLogicAnimNext/Source/RigLogicAnimNext/Public/RigLogicInstanceData.h
@@ -10,7 +10,7 @@
 struct FSharedRigRuntimeContext;
 struct FDNAIndexMapping;
 
-namespace UE::AnimNext
+namespace UE::UAF
 {
 	struct FRigLogicBoneMapping
 	{
@@ -45,10 +45,10 @@ namespace UE::AnimNext
 
 		uint32 NumLODs = 0;
 
-		void Init(const UE::AnimNext::FReferencePose* ReferencePose);
+		void Init(const UE::UAF::FReferencePose* ReferencePose);
 
 	private:
-		void InitBoneIndexMapping(const UE::AnimNext::FReferencePose* ReferencePose);
-		void InitSparseAndDenseDriverJointMapping(const UE::AnimNext::FReferencePose* ReferencePose);
+		void InitBoneIndexMapping(const UE::UAF::FReferencePose* ReferencePose);
+		void InitSparseAndDenseDriverJointMapping(const UE::UAF::FReferencePose* ReferencePose);
 	};
-} // namespace UE::AnimNext
\ No newline at end of file
+} // namespace UE::UAF
\ No newline at end of file
diff --git a/Engine/Plugins/Experimental/RigLogicAnimNext/Source/RigLogicAnimNext/Public/RigLogicInstanceDataPool.h b/Engine/Plugins/Experimental/RigLogicAnimNext/Source/RigLogicAnimNext/Public/RigLogicInstanceDataPool.h
--- a/Engine/Plugins/Experimental/RigLogicAnimNext/Source/RigLogicAnimNext/Public/RigLogicInstanceDataPool.h
+++ b/Engine/Plugins/Experimental/RigLogicAnimNext/Source/RigLogicAnimNext/Public/RigLogicInstanceDataPool.h
@@ -4,7 +4,7 @@
 
 #include "RigLogicInstanceData.h"
 
-namespace UE::AnimNext
+namespace UE::UAF
 {
 	class FRigLogicInstanceDataPool
 	{
@@ -13,10 +13,10 @@ namespace UE::AnimNext
 
 		mutable FCriticalSection PoolAccessCriticalSection;
 
-		TSharedPtr<FRigLogicAnimNextInstanceData> RequestData(const UE::AnimNext::FReferencePose* InReferencePose);
+		TSharedPtr<FRigLogicAnimNextInstanceData> RequestData(const UE::UAF::FReferencePose* InReferencePose);
 		void FreeData(TWeakObjectPtr<const USkeletalMesh> SkeletalMesh, TSharedPtr<FRigLogicAnimNextInstanceData> InData);
 
 		void GarbageCollect();
 		void Log();
 	};
-} // namespace UE::AnimNext
\ No newline at end of file
+} // namespace UE::UAF
\ No newline at end of file
diff --git a/Engine/Plugins/Experimental/RigLogicAnimNext/Source/RigLogicAnimNext/Public/RigLogicTask.h b/Engine/Plugins/Experimental/RigLogicAnimNext/Source/RigLogicAnimNext/Public/RigLogicTask.h
--- a/Engine/Plugins/Experimental/RigLogicAnimNext/Source/RigLogicAnimNext/Public/RigLogicTask.h
+++ b/Engine/Plugins/Experimental/RigLogicAnimNext/Source/RigLogicAnimNext/Public/RigLogicTask.h
@@ -27,10 +27,10 @@ struct RIGLOGICANIMNEXT_API FAnimNextRigLogicTask : public FAnimNextEvaluationTa
 
 	DECLARE_ANIM_EVALUATION_TASK(FAnimNextRigLogicTask)
 
-	static FAnimNextRigLogicTask Make(UE::AnimNext::FRigLogicTrait::FInstanceData* InstanceData);
+	static FAnimNextRigLogicTask Make(UE::UAF::FRigLogicTrait::FInstanceData* InstanceData);
 
 	// Task entry point
-	virtual void Execute(UE::AnimNext::FEvaluationVM& VM) const override;
+	virtual void Execute(UE::UAF::FEvaluationVM& VM) const override;
 
 	using FCachedIndexedCurve = TBaseBlendedCurve<FDefaultAllocator, UE::Anim::FCurveElementIndexed>;
 	static constexpr uint16 ATTR_COUNT_PER_JOINT = 10;
@@ -40,18 +40,18 @@ struct RIGLOGICANIMNEXT_API FAnimNextRigLogicTask : public FAnimNextEvaluationTa
 		const FCachedIndexedCurve& ControlAttributeCurves,
 		const FCachedIndexedCurve& NeuralNetworkMaskCurves,
 		const TArrayView<const float> NeutralJointValues,
-		const TArrayView<const UE::AnimNext::FPoseBoneControlAttributeMapping> SparseDriverJointsToControlAttributesMap,
-		const TArrayView<const UE::AnimNext::FPoseBoneControlAttributeMapping> DenseDriverJointsToControlAttributesMap,
-		const UE::AnimNext::FLODPoseStack& InputPose) const;
+		const TArrayView<const UE::UAF::FPoseBoneControlAttributeMapping> SparseDriverJointsToControlAttributesMap,
+		const TArrayView<const UE::UAF::FPoseBoneControlAttributeMapping> DenseDriverJointsToControlAttributesMap,
+		const UE::UAF::FLODPoseStack& InputPose) const;
 
-	void UpdateJoints(const TArrayView<const UE::AnimNext::FRigLogicBoneMapping> RigLogicIndexToMeshIndexMapping,
+	void UpdateJoints(const TArrayView<const UE::UAF::FRigLogicBoneMapping> RigLogicIndexToMeshIndexMapping,
 		const TArrayView<const uint16> VariableJointIndices,
 		const TArrayView<const float> NeutralJointValues,
 		const TArrayView<const float> DeltaJointValues,
-		UE::AnimNext::FLODPoseStack& OutputPose) const;
+		UE::UAF::FLODPoseStack& OutputPose) const;
 
 	void UpdateBlendShapeCurves(const FCachedIndexedCurve& MorphTargetCurves, const TArrayView<const float> BlendShapeValues, FBlendedCurve& OutputCurves) const;
 	void UpdateAnimMapCurves(const FCachedIndexedCurve& MaskMultiplierCurves, const TArrayView<const float> AnimMapOutputs, FBlendedCurve& OutputCurves) const;
 
-	UE::AnimNext::FRigLogicTrait::FInstanceData* TraitInstanceData = nullptr;
+	UE::UAF::FRigLogicTrait::FInstanceData* TraitInstanceData = nullptr;
 };
\ No newline at end of file
diff --git a/Engine/Source/Runtime/Engine/Classes/Animation/AnimSequence.h b/Engine/Source/Runtime/Engine/Classes/Animation/AnimSequence.h
--- a/Engine/Source/Runtime/Engine/Classes/Animation/AnimSequence.h
+++ b/Engine/Source/Runtime/Engine/Classes/Animation/AnimSequence.h
@@ -41,7 +41,7 @@ struct FAnimSequenceDecompressionContext;
 struct FCompactPose;
 
 namespace UE { namespace Anim { class FAnimSequenceCompilingManager; namespace Compression { struct FScopedCompressionGuard; } class FAnimationSequenceAsyncCacheTask; } }
-namespace UE::AnimNext { class FDecompressionTools; }
+namespace UE::UAF { class FDecompressionTools; }
 
 extern ENGINE_API int32 GPerformFrameStripping;
 
@@ -973,7 +973,7 @@ public:
 	friend struct UE::Anim::Compression::FScopedCompressionGuard;
 	friend class FAnimDataControllerTestBase;
 	friend class UE::Anim::FAnimSequenceCompilingManager;
-	friend UE::AnimNext::FDecompressionTools;
+	friend UE::UAF::FDecompressionTools;
 	friend class FAnimSequenceDetails;
 	friend class FCompressedAnimationDataNodeBuilder;
 	friend class UAnimBank;
diff --git a/Engine/Source/Runtime/Engine/Classes/Animation/AttributeCurve.h b/Engine/Source/Runtime/Engine/Classes/Animation/AttributeCurve.h
--- a/Engine/Source/Runtime/Engine/Classes/Animation/AttributeCurve.h
+++ b/Engine/Source/Runtime/Engine/Classes/Animation/AttributeCurve.h
@@ -14,7 +14,7 @@
 #include "AttributeCurve.generated.h"
 
 namespace UE { namespace Anim { class IAttributeBlendOperator; struct Attributes; } }
-namespace UE::AnimNext { class FDecompressionTools; }
+namespace UE::UAF { class FDecompressionTools; }
 
 typedef UE::Anim::TWrappedAttribute<FDefaultAllocator> FWrappedAttribute;
 
@@ -191,7 +191,7 @@ protected:
 
 	friend class UAnimSequence;
 	friend struct UE::Anim::Attributes;
-	friend UE::AnimNext::FDecompressionTools;
+	friend UE::UAF::FDecompressionTools;
 };
 
 template<>
diff --git a/Engine/Source/Runtime/Engine/Classes/Components/SkeletalMeshComponent.h b/Engine/Source/Runtime/Engine/Classes/Components/SkeletalMeshComponent.h
--- a/Engine/Source/Runtime/Engine/Classes/Components/SkeletalMeshComponent.h
+++ b/Engine/Source/Runtime/Engine/Classes/Components/SkeletalMeshComponent.h
@@ -42,7 +42,7 @@ struct FConstraintProfileProperties;
 struct FNavigableGeometryExport;
 struct FCompactPose;
 
-namespace UE::AnimNext
+namespace UE::UAF
 {
 	class FDataRegistry;	// Temporary
 }
@@ -2725,7 +2725,7 @@ private:
 	/** Register/Unregister for OnLODRequiredBonesUpdate callback */
 	ENGINE_API FDelegateHandle RegisterOnLODRequiredBonesUpdate_Member(const FOnLODRequiredBonesUpdate& Delegate);
 	ENGINE_API void UnregisterOnLODRequiredBonesUpdate_Member(const FDelegateHandle& DelegateHandle);
-	friend class UE::AnimNext::FDataRegistry;
+	friend class UE::UAF::FDataRegistry;
 	
 private:
 
